[Sequencer] Move Navigation Tool to new experimental SequenceNavigator plugin

diff --git a/Engine/Plugins/Experimental/Avalanche/Source/AvalancheSequencer/AvalancheSequencer.Build.cs b/Engine/Plugins/Experimental/Avalanche/Source/AvalancheSequencer/AvalancheSequencer.Build.cs
--- a/Engine/Plugins/Experimental/Avalanche/Source/AvalancheSequencer/AvalancheSequencer.Build.cs
+++ b/Engine/Plugins/Experimental/Avalanche/Source/AvalancheSequencer/AvalancheSequencer.Build.cs
@@ -50,6 +50,7 @@ public class AvalancheSequencer : ModuleRules
 				"Projects",
 				"SceneOutliner",
 				"SequencerCore",
+				"SequenceNavigator",
 				"Settings",
 				"Slate",
 				"StateTreeEditorModule",
diff --git a/Engine/Plugins/Experimental/Avalanche/Source/AvalancheSequencer/Private/AvaSequencer.cpp b/Engine/Plugins/Experimental/Avalanche/Source/AvalancheSequencer/Private/AvaSequencer.cpp
--- a/Engine/Plugins/Experimental/Avalanche/Source/AvalancheSequencer/Private/AvaSequencer.cpp
+++ b/Engine/Plugins/Experimental/Avalanche/Source/AvalancheSequencer/Private/AvaSequencer.cpp
@@ -30,6 +30,7 @@
 #include "Framework/Notifications/NotificationManager.h"
 #include "IAvaSequenceProvider.h"
 #include "IContentBrowserSingleton.h"
+#include "INavigationTool.h"
 #include "ISequencer.h"
 #include "ISequencerModule.h"
 #include "ISequencerTrackEditor.h"
@@ -40,6 +41,7 @@
 #include "Misc/MessageDialog.h"
 #include "Misc/NotifyHook.h"
 #include "MovieScene.h"
+#include "NavigationToolDefines.h"
 #include "NavigationTool/AvaNavigationToolProvider.h"
 #include "NavigationTool/NavigationToolAvaSequence.h"
 #include "Playback/AvaSequencerCleanView.h"
@@ -59,8 +61,6 @@
 #include "Sidebar/SidebarDrawerConfig.h"
 #include "StaggerTool/AvaStaggerTool.h"
 #include "Toolkits/AssetEditorToolkit.h"
-#include "Tools/NavigationTool/INavigationTool.h"
-#include "Tools/NavigationTool/NavigationToolDefines.h"
 #include "Widgets/Notifications/SNotificationList.h"
 #include "Widgets/Views/SHeaderRow.h"
 
@@ -448,7 +448,6 @@ TSharedRef<ISequencer> FAvaSequencer::CreateSequencer()
 		SequencerInitParams.HostCapabilities.bSupportsCurveEditor = true;
 		SequencerInitParams.HostCapabilities.bSupportsSaveMovieSceneAsset = false;
 		SequencerInitParams.HostCapabilities.bSupportsSidebar = true;
-		SequencerInitParams.HostCapabilities.bSupportsNavigationTool = true;
 	};
 
 	InstancedSequencer = FAvaSequencerUtils::GetSequencerModule().CreateSequencer(SequencerInitParams);
@@ -991,7 +990,7 @@ void FAvaSequencer::ApplyPresetToSelection(const FAvaSequencePreset& InPreset)
 		return;
 	}
 
-	const TSharedPtr<INavigationTool> Tool = Sequencer->GetNavigationTool();
+	const TSharedPtr<INavigationTool> Tool = FNavigationToolExtender::FindNavigationTool(Sequencer.ToSharedRef());
 	if (!Tool.IsValid())
 	{
 		return;
@@ -2082,9 +2081,9 @@ void FAvaSequencer::NotifyOnSequenceTreeChanged()
 {
 	using namespace UE::Sequencer::NavigationTool;
 
-	if (const TSharedPtr<INavigationTool> NavigationTool = GetSequencer()->GetNavigationTool())
+	if (const TSharedPtr<INavigationTool> Tool = FNavigationToolExtender::FindNavigationTool(GetSequencer()))
 	{
-		NavigationTool->RequestRefresh();
+		Tool->RequestRefresh();
 	}
 
 	IAvaSequenceProvider* const SequenceProvider = Provider.GetSequenceProvider();
@@ -2178,7 +2177,7 @@ TArray<TSharedPtr<UE::Sequencer::NavigationTool::FNavigationToolAvaSequence>> FA
 		return {};
 	}
 
-	const TSharedPtr<INavigationTool> Tool = Sequencer->GetNavigationTool();
+	const TSharedPtr<INavigationTool> Tool = FNavigationToolExtender::FindNavigationTool(Sequencer.ToSharedRef());
 	if (!Tool.IsValid())
 	{
 		return {};
@@ -2410,25 +2409,20 @@ void FAvaSequencer::RegisterNavigationTool()
 {
 	using namespace UE::Sequencer::NavigationTool;
 
-	const TSharedPtr<ISequencer> Sequencer = SequencerWeak.Pin();
-	if (!Sequencer.IsValid())
-	{
-		return;
-	}
-
-	const TSharedPtr<INavigationTool> Tool = Sequencer->GetNavigationTool();
-	if (!Tool.IsValid())
+	if (!SequencerWeak.IsValid())
 	{
 		return;
 	}
 
-	const FName ToolId = FNavigationToolExtender::GetToolInstanceId(*Tool);
+	const TSharedRef<ISequencer> Sequencer = SequencerWeak.Pin().ToSharedRef();
+	const FName ToolId = FNavigationToolExtender::GetToolInstanceId(*Sequencer);
 
-	const TSharedPtr<FNavigationToolProvider> FoundProvider = FNavigationToolExtender::FindToolProvider(ToolId, UE::AvaSequencer::FAvaNavigationToolProvider::Identifier);
+	const TSharedPtr<FNavigationToolProvider> FoundProvider = FNavigationToolExtender::FindToolProvider(ToolId
+		, UE::AvaSequencer::FAvaNavigationToolProvider::Identifier);
 	if (!FoundProvider.IsValid())
 	{
 		NavigationToolProvider = MakeShared<UE::AvaSequencer::FAvaNavigationToolProvider>(SharedThis(this));
-		FNavigationToolExtender::RegisterToolProvider(ToolId, NavigationToolProvider.ToSharedRef());
+		FNavigationToolExtender::RegisterToolProvider(Sequencer, NavigationToolProvider.ToSharedRef());
 	}
 }
 
@@ -2447,7 +2441,7 @@ void FAvaSequencer::UnregisterNavigationTool()
 		return;
 	}
 
-	const TSharedPtr<INavigationTool> Tool = Sequencer->GetNavigationTool();
+	const TSharedPtr<INavigationTool> Tool = FNavigationToolExtender::FindNavigationTool(Sequencer.ToSharedRef());
 	if (!Tool.IsValid())
 	{
 		return;
diff --git a/Engine/Plugins/Experimental/Avalanche/Source/AvalancheSequencer/Private/NavigationTool/AvaNavigationToolProvider.cpp b/Engine/Plugins/Experimental/Avalanche/Source/AvalancheSequencer/Private/NavigationTool/AvaNavigationToolProvider.cpp
--- a/Engine/Plugins/Experimental/Avalanche/Source/AvalancheSequencer/Private/NavigationTool/AvaNavigationToolProvider.cpp
+++ b/Engine/Plugins/Experimental/Avalanche/Source/AvalancheSequencer/Private/NavigationTool/AvaNavigationToolProvider.cpp
@@ -16,27 +16,27 @@
 #include "NavigationTool/Columns/AvaNavigationToolStatusColumn.h"
 #include "NavigationTool/NavigationToolAvaSequence.h"
 #include "Settings/AvaSequencerSettings.h"
-#include "Tools/NavigationTool/Columns/NavigationToolColorColumn.h"
-#include "Tools/NavigationTool/Columns/NavigationToolColumnExtender.h"
-#include "Tools/NavigationTool/Columns/NavigationToolCommentColumn.h"
-#include "Tools/NavigationTool/Columns/NavigationToolHBiasColumn.h"
-#include "Tools/NavigationTool/Columns/NavigationToolInTimeColumn.h"
-#include "Tools/NavigationTool/Columns/NavigationToolItemsColumn.h"
-#include "Tools/NavigationTool/Columns/NavigationToolLabelColumn.h"
-#include "Tools/NavigationTool/Columns/NavigationToolLengthColumn.h"
-#include "Tools/NavigationTool/Columns/NavigationToolLockColumn.h"
-#include "Tools/NavigationTool/Columns/NavigationToolMarkerVisibilityColumn.h"
-#include "Tools/NavigationTool/Columns/NavigationToolOutTimeColumn.h"
-#include "Tools/NavigationTool/Columns/NavigationToolPlayheadColumn.h"
-#include "Tools/NavigationTool/Columns/NavigationToolRevisionControlColumn.h"
-#include "Tools/NavigationTool/Columns/NavigationToolStartFrameOffsetColumn.h"
-#include "Tools/NavigationTool/Columns/NavigationToolTakeColumn.h"
-#include "Tools/NavigationTool/Columns/NavigationToolDeactiveStateColumn.h"
-#include "Tools/NavigationTool/Filters/Filters/NavigationToolBuiltInFilterParams.h"
-#include "Tools/NavigationTool/INavigationTool.h"
-#include "Tools/NavigationTool/ItemActions/NavigationToolAddItem.h"
-#include "Tools/NavigationTool/Items/NavigationToolItemParameters.h"
-#include "Tools/NavigationTool/NavigationToolSettings.h"
+#include "Columns/NavigationToolColorColumn.h"
+#include "Columns/NavigationToolColumnExtender.h"
+#include "Columns/NavigationToolCommentColumn.h"
+#include "Columns/NavigationToolDeactiveStateColumn.h"
+#include "Columns/NavigationToolHBiasColumn.h"
+#include "Columns/NavigationToolInTimeColumn.h"
+#include "Columns/NavigationToolItemsColumn.h"
+#include "Columns/NavigationToolLabelColumn.h"
+#include "Columns/NavigationToolLengthColumn.h"
+#include "Columns/NavigationToolLockColumn.h"
+#include "Columns/NavigationToolMarkerVisibilityColumn.h"
+#include "Columns/NavigationToolOutTimeColumn.h"
+#include "Columns/NavigationToolPlayheadColumn.h"
+#include "Columns/NavigationToolRevisionControlColumn.h"
+#include "Columns/NavigationToolStartFrameOffsetColumn.h"
+#include "Columns/NavigationToolTakeColumn.h"
+#include "Filters/Filters/NavigationToolBuiltInFilterParams.h"
+#include "INavigationTool.h"
+#include "ItemActions/NavigationToolAddItem.h"
+#include "Items/NavigationToolItemParameters.h"
+#include "NavigationToolSettings.h"
 
 #define LOCTEXT_NAMESPACE "AvaNavigationToolProvider"
 
@@ -397,6 +397,8 @@ void FAvaNavigationToolProvider::ExtendToolItemDragDropOp(UE::Sequencer::Navigat
 
 void FAvaNavigationToolProvider::OnSequenceAdded(UAvaSequence* const InAvaSequence)
 {
+	using namespace UE::Sequencer::NavigationTool;
+
 	const TSharedPtr<FAvaSequencer> AvaSequencer = StaticCastSharedPtr<FAvaSequencer>(WeakAvaSequencer.Pin());
 	if (!AvaSequencer.IsValid())
 	{
@@ -409,7 +411,7 @@ void FAvaNavigationToolProvider::OnSequenceAdded(UAvaSequence* const InAvaSequen
 		return;
 	}
 
-	const TSharedPtr<UE::Sequencer::NavigationTool::INavigationTool> NavigationTool = Sequencer->GetNavigationTool();
+	const TSharedPtr<INavigationTool> NavigationTool = FNavigationToolExtender::FindNavigationTool(Sequencer.ToSharedRef());
 	if (!NavigationTool.IsValid())
 	{
 		return;
@@ -492,7 +494,7 @@ bool FAvaNavigationToolProvider::CanRelabelSequence() const
 		return false;
 	}
 
-	const TSharedPtr<INavigationTool> NavigationTool = AvaSequencer->GetSequencer()->GetNavigationTool();
+	const TSharedPtr<INavigationTool> NavigationTool = FNavigationToolExtender::FindNavigationTool(AvaSequencer->GetSequencer());
 	if (!NavigationTool.IsValid())
 	{
 		return false;
@@ -513,7 +515,7 @@ void FAvaNavigationToolProvider::RelabelSequence()
 		return;
 	}
 
-	const TSharedPtr<INavigationTool> NavigationTool = AvaSequencer->GetSequencer()->GetNavigationTool();
+	const TSharedPtr<INavigationTool> NavigationTool = FNavigationToolExtender::FindNavigationTool(AvaSequencer->GetSequencer());
 	if (!NavigationTool.IsValid())
 	{
 		return;
@@ -545,7 +547,7 @@ bool FAvaNavigationToolProvider::CanAddSequenceToSelection() const
 		return false;
 	}
 
-	const TSharedPtr<INavigationTool> NavigationTool = AvaSequencer->GetSequencer()->GetNavigationTool();
+	const TSharedPtr<INavigationTool> NavigationTool = FNavigationToolExtender::FindNavigationTool(AvaSequencer->GetSequencer());
 	if (!NavigationTool.IsValid())
 	{
 		return false;
@@ -566,7 +568,7 @@ void FAvaNavigationToolProvider::AddSequenceToSelection()
 		return;
 	}
 
-	const TSharedPtr<INavigationTool> NavigationTool = AvaSequencer->GetSequencer()->GetNavigationTool();
+	const TSharedPtr<INavigationTool> NavigationTool = FNavigationToolExtender::FindNavigationTool(AvaSequencer->GetSequencer());
 	if (!NavigationTool.IsValid())
 	{
 		return;
diff --git a/Engine/Plugins/Experimental/Avalanche/Source/AvalancheSequencer/Private/NavigationTool/AvaNavigationToolProvider.h b/Engine/Plugins/Experimental/Avalanche/Source/AvalancheSequencer/Private/NavigationTool/AvaNavigationToolProvider.h
--- a/Engine/Plugins/Experimental/Avalanche/Source/AvalancheSequencer/Private/NavigationTool/AvaNavigationToolProvider.h
+++ b/Engine/Plugins/Experimental/Avalanche/Source/AvalancheSequencer/Private/NavigationTool/AvaNavigationToolProvider.h
@@ -3,9 +3,9 @@
 #pragma once
 
 #include "Delegates/IDelegateInstance.h"
+#include "NavigationToolDefines.h"
+#include "Providers/NavigationToolProvider.h"
 #include "Templates/SharedPointer.h"
-#include "Tools/NavigationTool/NavigationToolDefines.h"
-#include "Tools/NavigationTool/Providers/NavigationToolProvider.h"
 #include "UObject/WeakObjectPtr.h"
 
 class FAvaSequencer;
diff --git a/Engine/Plugins/Experimental/Avalanche/Source/AvalancheSequencer/Private/NavigationTool/Columns/AvaNavigationToolStatusColumn.cpp b/Engine/Plugins/Experimental/Avalanche/Source/AvalancheSequencer/Private/NavigationTool/Columns/AvaNavigationToolStatusColumn.cpp
--- a/Engine/Plugins/Experimental/Avalanche/Source/AvalancheSequencer/Private/NavigationTool/Columns/AvaNavigationToolStatusColumn.cpp
+++ b/Engine/Plugins/Experimental/Avalanche/Source/AvalancheSequencer/Private/NavigationTool/Columns/AvaNavigationToolStatusColumn.cpp
@@ -5,9 +5,8 @@
 #include "NavigationTool/NavigationToolAvaSequence.h"
 #include "NavigationTool/Columns/SAvaNavigationToolStatus.h"
 #include "Styling/AppStyle.h"
-#include "Tools/NavigationTool/INavigationTool.h"
-#include "Tools/NavigationTool/INavigationToolView.h"
-#include "Tools/NavigationTool/NavigationToolDefines.h"
+#include "INavigationTool.h"
+#include "INavigationToolView.h"
 
 #define LOCTEXT_NAMESPACE "AvaNavigationToolStatusColumn"
 
diff --git a/Engine/Plugins/Experimental/Avalanche/Source/AvalancheSequencer/Private/NavigationTool/Columns/AvaNavigationToolStatusColumn.h b/Engine/Plugins/Experimental/Avalanche/Source/AvalancheSequencer/Private/NavigationTool/Columns/AvaNavigationToolStatusColumn.h
--- a/Engine/Plugins/Experimental/Avalanche/Source/AvalancheSequencer/Private/NavigationTool/Columns/AvaNavigationToolStatusColumn.h
+++ b/Engine/Plugins/Experimental/Avalanche/Source/AvalancheSequencer/Private/NavigationTool/Columns/AvaNavigationToolStatusColumn.h
@@ -2,7 +2,7 @@
 
 #pragma once
 
-#include "Tools/NavigationTool/Columns/NavigationToolColumn.h"
+#include "Columns/NavigationToolColumn.h"
 
 namespace UE::Sequencer::NavigationTool
 {
diff --git a/Engine/Plugins/Experimental/Avalanche/Source/AvalancheSequencer/Private/NavigationTool/Columns/SAvaNavigationToolStatus.cpp b/Engine/Plugins/Experimental/Avalanche/Source/AvalancheSequencer/Private/NavigationTool/Columns/SAvaNavigationToolStatus.cpp
--- a/Engine/Plugins/Experimental/Avalanche/Source/AvalancheSequencer/Private/NavigationTool/Columns/SAvaNavigationToolStatus.cpp
+++ b/Engine/Plugins/Experimental/Avalanche/Source/AvalancheSequencer/Private/NavigationTool/Columns/SAvaNavigationToolStatus.cpp
@@ -6,12 +6,12 @@
 #include "AvaSequencerUtils.h"
 #include "IAvaSceneInterface.h"
 #include "IAvaSequencerProvider.h"
+#include "INavigationToolView.h"
 #include "ISequencer.h"
 #include "NavigationTool/NavigationToolAvaSequence.h"
-#include "Tools/NavigationTool/INavigationToolView.h"
-#include "Tools/NavigationTool/Widgets/SNavigationToolTreeRow.h"
 #include "Widgets/Layout/SScrollBox.h"
 #include "Widgets/Notifications/SProgressBar.h"
+#include "Widgets/SNavigationToolTreeRow.h"
 
 #define LOCTEXT_NAMESPACE "SAvaNavigationToolStatus"
 
diff --git a/Engine/Plugins/Experimental/Avalanche/Source/AvalancheSequencer/Private/NavigationTool/Columns/SAvaNavigationToolStatus.h b/Engine/Plugins/Experimental/Avalanche/Source/AvalancheSequencer/Private/NavigationTool/Columns/SAvaNavigationToolStatus.h
--- a/Engine/Plugins/Experimental/Avalanche/Source/AvalancheSequencer/Private/NavigationTool/Columns/SAvaNavigationToolStatus.h
+++ b/Engine/Plugins/Experimental/Avalanche/Source/AvalancheSequencer/Private/NavigationTool/Columns/SAvaNavigationToolStatus.h
@@ -3,7 +3,7 @@
 #pragma once
 
 #include "Misc/FrameTime.h"
-#include "Tools/NavigationTool/NavigationToolDefines.h"
+#include "NavigationToolDefines.h"
 #include "Widgets/SCompoundWidget.h"
 
 class FText;
diff --git a/Engine/Plugins/Experimental/Avalanche/Source/AvalancheSequencer/Private/NavigationTool/DragDrop/NavigationToolAvaSequenceDropHandler.cpp b/Engine/Plugins/Experimental/Avalanche/Source/AvalancheSequencer/Private/NavigationTool/DragDrop/NavigationToolAvaSequenceDropHandler.cpp
--- a/Engine/Plugins/Experimental/Avalanche/Source/AvalancheSequencer/Private/NavigationTool/DragDrop/NavigationToolAvaSequenceDropHandler.cpp
+++ b/Engine/Plugins/Experimental/Avalanche/Source/AvalancheSequencer/Private/NavigationTool/DragDrop/NavigationToolAvaSequenceDropHandler.cpp
@@ -5,12 +5,12 @@
 #include "IAvaSequenceProvider.h"
 #include "IAvaSequencer.h"
 #include "IAvaSequencerProvider.h"
+#include "INavigationTool.h"
+#include "ItemActions/NavigationToolAddItem.h"
+#include "ItemActions/NavigationToolRemoveItem.h"
 #include "Misc/Optional.h"
 #include "NavigationTool/NavigationToolAvaSequence.h"
 #include "ScopedTransaction.h"
-#include "Tools/NavigationTool/INavigationTool.h"
-#include "Tools/NavigationTool/ItemActions/NavigationToolAddItem.h"
-#include "Tools/NavigationTool/ItemActions/NavigationToolRemoveItem.h"
 
 #define LOCTEXT_NAMESPACE "NavigationToolAvaSequenceDropHandler"
 
diff --git a/Engine/Plugins/Experimental/Avalanche/Source/AvalancheSequencer/Private/NavigationTool/DragDrop/NavigationToolAvaSequenceDropHandler.h b/Engine/Plugins/Experimental/Avalanche/Source/AvalancheSequencer/Private/NavigationTool/DragDrop/NavigationToolAvaSequenceDropHandler.h
--- a/Engine/Plugins/Experimental/Avalanche/Source/AvalancheSequencer/Private/NavigationTool/DragDrop/NavigationToolAvaSequenceDropHandler.h
+++ b/Engine/Plugins/Experimental/Avalanche/Source/AvalancheSequencer/Private/NavigationTool/DragDrop/NavigationToolAvaSequenceDropHandler.h
@@ -3,7 +3,7 @@
 #pragma once
 
 #include "Templates/SharedPointer.h"
-#include "Tools/NavigationTool/DragDropOps/Handlers/NavigationToolItemDropHandler.h"
+#include "DragDropOps/Handlers/NavigationToolItemDropHandler.h"
 
 class IAvaSequencer;
 enum class EItemDropZone;
diff --git a/Engine/Plugins/Experimental/Avalanche/Source/AvalancheSequencer/Private/NavigationTool/NavigationToolAvaSequence.cpp b/Engine/Plugins/Experimental/Avalanche/Source/AvalancheSequencer/Private/NavigationTool/NavigationToolAvaSequence.cpp
--- a/Engine/Plugins/Experimental/Avalanche/Source/AvalancheSequencer/Private/NavigationTool/NavigationToolAvaSequence.cpp
+++ b/Engine/Plugins/Experimental/Avalanche/Source/AvalancheSequencer/Private/NavigationTool/NavigationToolAvaSequence.cpp
@@ -6,12 +6,12 @@
 #include "AvaSequencer.h"
 #include "AvaSequencerUtils.h"
 #include "IAvaSequencer.h"
+#include "INavigationTool.h"
 #include "ISequencer.h"
+#include "Items/NavigationToolItemParameters.h"
+#include "Items/NavigationToolSequence.h"
 #include "ScopedTransaction.h"
 #include "Styling/StyleColors.h"
-#include "Tools/NavigationTool/INavigationTool.h"
-#include "Tools/NavigationTool/Items/NavigationToolItemParameters.h"
-#include "Tools/NavigationTool/Items/NavigationToolSequence.h"
 
 #define LOCTEXT_NAMESPACE "NavigationToolAvaSequence"
 
diff --git a/Engine/Plugins/Experimental/Avalanche/Source/AvalancheSequencer/Private/NavigationTool/NavigationToolAvaSequence.h b/Engine/Plugins/Experimental/Avalanche/Source/AvalancheSequencer/Private/NavigationTool/NavigationToolAvaSequence.h
--- a/Engine/Plugins/Experimental/Avalanche/Source/AvalancheSequencer/Private/NavigationTool/NavigationToolAvaSequence.h
+++ b/Engine/Plugins/Experimental/Avalanche/Source/AvalancheSequencer/Private/NavigationTool/NavigationToolAvaSequence.h
@@ -2,7 +2,7 @@
 
 #pragma once
 
-#include "Tools/NavigationTool/Items/NavigationToolSequence.h"
+#include "Items/NavigationToolSequence.h"
 
 class UAvaSequence;
 
diff --git a/Engine/Plugins/Experimental/SequenceNavigator/Source/SequenceNavigator/Private/Columns/NavigationToolColumn.cpp b/Engine/Plugins/Experimental/SequenceNavigator/Source/SequenceNavigator/Private/Columns/NavigationToolColumn.cpp
new file mode 100644
--- /dev/null
+++ b/Engine/Plugins/Experimental/SequenceNavigator/Source/SequenceNavigator/Private/Columns/NavigationToolColumn.cpp
@@ -0,0 +1,8 @@
+﻿// Copyright Epic Games, Inc. All Rights Reserved.
+
+#include "Columns/NavigationToolColumn.h"
+
+namespace UE::Sequencer::NavigationTool
+{
+
+}
diff --git a/Engine/Plugins/Experimental/SequenceNavigator/Source/SequenceNavigator/Private/DragDropOps/Handlers/NavigationToolItemDropHandler.cpp b/Engine/Plugins/Experimental/SequenceNavigator/Source/SequenceNavigator/Private/DragDropOps/Handlers/NavigationToolItemDropHandler.cpp
new file mode 100644
--- /dev/null
+++ b/Engine/Plugins/Experimental/SequenceNavigator/Source/SequenceNavigator/Private/DragDropOps/Handlers/NavigationToolItemDropHandler.cpp
@@ -0,0 +1,21 @@
+﻿// Copyright Epic Games, Inc. All Rights Reserved.
+
+#include "DragDropOps/Handlers/NavigationToolItemDropHandler.h"
+#include "DragDropOps/NavigationToolItemDragDropOp.h"
+
+namespace UE::Sequencer::NavigationTool
+{
+
+void FNavigationToolItemDropHandler::Initialize(const FNavigationToolItemDragDropOp& InDragDropOp)
+{
+	Items = InDragDropOp.GetItems();
+	ActionType = InDragDropOp.GetActionType();
+
+	// Remove all Items that are Invalid or out of the Scope of this Handler
+	Items.RemoveAll([this](const FNavigationToolItemPtr& InItem)
+		{
+			return !InItem.IsValid() || !IsDraggedItemSupported(InItem);
+		});
+}
+
+}
diff --git a/Engine/Plugins/Experimental/SequenceNavigator/Source/SequenceNavigator/Private/DragDropOps/NavigationToolItemDragDropOp.cpp b/Engine/Plugins/Experimental/SequenceNavigator/Source/SequenceNavigator/Private/DragDropOps/NavigationToolItemDragDropOp.cpp
new file mode 100644
--- /dev/null
+++ b/Engine/Plugins/Experimental/SequenceNavigator/Source/SequenceNavigator/Private/DragDropOps/NavigationToolItemDragDropOp.cpp
@@ -0,0 +1,142 @@
+// Copyright Epic Games, Inc. All Rights Reserved.
+
+#include "DragDropOps/NavigationToolItemDragDropOp.h"
+#include "DragDropOps/Handlers/NavigationToolSequenceDropHandler.h"
+#include "DragDropOps/Handlers/NavigationToolItemDropHandler.h"
+#include "GameFramework/Actor.h"
+#include "INavigationToolView.h"
+#include "Items/INavigationToolItem.h"
+#include "LevelSequenceActor.h"
+#include "NavigationToolExtender.h"
+#include "NavigationToolView.h"
+#include "ScopedTransaction.h"
+#include "Styling/SlateIconFinder.h"
+
+#define LOCTEXT_NAMESPACE "NavigationToolItemDragDropOp"
+
+namespace UE::Sequencer::NavigationTool
+{
+
+FText GetActionName(const ENavigationToolDragDropActionType InActionType)
+{
+	switch (InActionType)
+	{
+		case ENavigationToolDragDropActionType::Move:
+			return LOCTEXT("MoveAction", "Moving");
+
+		case ENavigationToolDragDropActionType::Copy:
+			return LOCTEXT("CopyAction", "Copying");
+	}
+	return FText::GetEmpty();
+}
+
+FText GetItemName(const TArray<FNavigationToolItemPtr>& InItems)
+{
+	if (InItems.IsEmpty())
+	{
+		return LOCTEXT("NoItems", "0 Items");
+	}
+
+	if (InItems.Num() == 1)
+	{
+		return InItems[0]->GetDisplayName();
+	}
+
+	return FText::Format(LOCTEXT("ManyItems", "{0} and {1} other item(s)"), InItems[0]->GetDisplayName(), InItems.Num() - 1);
+}
+
+TSharedRef<FNavigationToolItemDragDropOp> FNavigationToolItemDragDropOp::New(const TArray<FNavigationToolItemPtr>& InItems
+	, const TSharedPtr<FNavigationToolView>& InToolView
+	, ENavigationToolDragDropActionType InActionType)
+{
+	TSharedRef<FNavigationToolItemDragDropOp> DragDropOp = MakeShared<FNavigationToolItemDragDropOp>();
+	DragDropOp->Init(InItems, InToolView, InActionType);
+	return DragDropOp;
+}
+
+FReply FNavigationToolItemDragDropOp::Drop(EItemDropZone InDropZone, const FNavigationToolItemPtr& InTargetItem)
+{
+	FReply Reply = FReply::Unhandled();
+
+	FScopedTransaction Transaction(LOCTEXT("DropItems", "Navigation Tool Drop Items"));
+
+	for (const TSharedRef<FNavigationToolItemDropHandler>& DropHandler : DropHandlers)
+	{
+		if (DropHandler->GetItems().IsEmpty())
+		{
+			continue;
+		}
+
+		if (DropHandler->Drop(InDropZone, InTargetItem))
+		{
+			Reply = FReply::Handled();
+		}
+	}
+
+	if (!Reply.IsEventHandled())
+	{
+		Transaction.Cancel();
+	}
+
+	return Reply;
+}
+
+TOptional<EItemDropZone> FNavigationToolItemDragDropOp::CanDrop(const EItemDropZone InDropZone, const FNavigationToolItemPtr& InTargetItem) const
+{
+	const TSharedPtr<INavigationToolView> ToolView = GetToolView();
+
+	// Only support Drag/Drop from Same Navigation Tool
+	if (!ToolView.IsValid() || &InTargetItem->GetOwnerTool() != ToolView->GetOwnerTool().Get())
+	{
+		return TOptional<EItemDropZone>();
+	}
+
+	for (const TSharedRef<FNavigationToolItemDropHandler>& DropHandler : DropHandlers)
+	{
+		if (DropHandler->GetItems().IsEmpty())
+		{
+			continue;
+		}
+
+		// Return the Item Drop Zone of the first Drop Handler that supports the Drop Zone and Target Item
+		const TOptional<EItemDropZone> DropZone = DropHandler->CanDrop(InDropZone, InTargetItem);
+		if (DropZone.IsSet())
+		{
+			return *DropZone;
+		}
+	}
+
+	return TOptional<EItemDropZone>();
+}
+
+void FNavigationToolItemDragDropOp::Init(const TArray<FNavigationToolItemPtr>& InItems
+	, const TSharedPtr<FNavigationToolView>& InToolView
+	, ENavigationToolDragDropActionType InActionType)
+{
+	Items = InItems;
+	WeakToolView = InToolView;
+	ActionType = InActionType;
+	MouseCursor = EMouseCursor::GrabHandClosed;
+
+	CurrentIconBrush = Items.Num() == 1
+		? Items[0]->GetIconBrush()
+		: FSlateIconFinder::FindIconForClass(ALevelSequenceActor::StaticClass()).GetIcon();
+
+	CurrentHoverText = FText::Format(LOCTEXT("HoverText", "{0} {1}")
+		, GetActionName(InActionType)
+		, GetItemName(InItems));
+
+	CurrentIconColorAndOpacity = FSlateColor::UseForeground();
+
+	SetupDefaults();
+	Construct();
+
+	// Add Default Drop Handlers 
+	AddDropHandler<FNavigationToolSequenceDropHandler>();
+
+	FNavigationToolExtender::OnItemDragDropOpInitialized().Broadcast(*this);
+}
+
+}
+
+#undef LOCTEXT_NAMESPACE
diff --git a/Engine/Plugins/Experimental/SequenceNavigator/Source/SequenceNavigator/Private/Filters/Filters/NavigationToolFilterBase.cpp b/Engine/Plugins/Experimental/SequenceNavigator/Source/SequenceNavigator/Private/Filters/Filters/NavigationToolFilterBase.cpp
new file mode 100644
--- /dev/null
+++ b/Engine/Plugins/Experimental/SequenceNavigator/Source/SequenceNavigator/Private/Filters/Filters/NavigationToolFilterBase.cpp
@@ -0,0 +1,22 @@
+﻿// Copyright Epic Games, Inc. All Rights Reserved.
+
+#include "Filters/Filters/NavigationToolFilterBase.h"
+#include "Filters/INavigationToolFilterBar.h"
+#include "Filters/ISequencerFilterBar.h"
+
+namespace UE::Sequencer::NavigationTool
+{
+
+using namespace UE::Sequencer;
+
+FNavigationToolFilter::FNavigationToolFilter(INavigationToolFilterBar& InOutFilterInterface, TSharedPtr<FFilterCategory>&& InCategory)
+	: FSequencerFilterBase<FNavigationToolItemPtr>(InOutFilterInterface, MoveTemp(InCategory))
+{
+}
+
+INavigationToolFilterBar& FNavigationToolFilter::GetFilterInterface() const
+{
+	return static_cast<INavigationToolFilterBar&>(FilterInterface);
+}
+
+}
diff --git a/Engine/Plugins/Experimental/SequenceNavigator/Source/SequenceNavigator/Private/Filters/NavigationToolFilterData.cpp b/Engine/Plugins/Experimental/SequenceNavigator/Source/SequenceNavigator/Private/Filters/NavigationToolFilterData.cpp
new file mode 100644
--- /dev/null
+++ b/Engine/Plugins/Experimental/SequenceNavigator/Source/SequenceNavigator/Private/Filters/NavigationToolFilterData.cpp
@@ -0,0 +1,142 @@
+﻿// Copyright Epic Games, Inc. All Rights Reserved.
+
+#include "Filters/NavigationToolFilterData.h"
+#include "Items/INavigationToolItem.h"
+
+namespace UE::Sequencer::NavigationTool
+{
+
+using namespace UE::Sequencer;
+
+FNavigationToolFilterData::FNavigationToolFilterData(const FString& InRawFilterText)
+	: RawFilterText(InRawFilterText)
+{
+}
+
+bool FNavigationToolFilterData::operator==(const FNavigationToolFilterData& InRhs) const
+{
+	return ContainsFilterInNodes(InRhs) && TotalNodeCount == InRhs.GetTotalNodeCount();
+}
+
+bool FNavigationToolFilterData::operator!=(const FNavigationToolFilterData& InRhs) const
+{
+	return !(*this == InRhs);
+}
+
+void FNavigationToolFilterData::Reset()
+{
+	RawFilterText.Reset();
+	TotalNodeCount = 0;
+	FilterInNodes.Reset();
+}
+
+FString FNavigationToolFilterData::GetRawFilterText() const
+{
+	return RawFilterText;
+}
+
+uint32 FNavigationToolFilterData::GetDisplayNodeCount() const
+{
+	return FilterInNodes.Num();
+}
+
+uint32 FNavigationToolFilterData::GetTotalNodeCount() const
+{
+	return TotalNodeCount;
+}
+
+uint32 FNavigationToolFilterData::GetFilterInCount() const
+{
+	return FilterInNodes.Num();
+}
+
+uint32 FNavigationToolFilterData::GetFilterOutCount() const
+{
+	return GetTotalNodeCount() - GetFilterInCount();
+}
+
+void FNavigationToolFilterData::IncrementTotalNodeCount()
+{
+	++TotalNodeCount;
+}
+
+void FNavigationToolFilterData::FilterInNode(const FNavigationToolItemPtr& InNode)
+{
+	FilterInNodes.Add(InNode);
+}
+
+void FNavigationToolFilterData::FilterOutNode(const FNavigationToolItemPtr& InNode)
+{
+	const FSetElementId ElementId = FilterInNodes.FindId(InNode);
+	if (ElementId.IsValidId())
+	{
+		FilterInNodes.Remove(ElementId);
+	}
+}
+
+void FNavigationToolFilterData::FilterInParentChildNodes(const FNavigationToolItemPtr& InNode
+	, const bool bInIncludeSelf, const bool bInIncludeParents, const bool bInIncludeChildren)
+{
+	if (!InNode.IsValid())
+	{
+		return;
+	}
+
+	if (bInIncludeParents)
+	{
+		FNavigationToolItemPtr CurrentParent = InNode->GetParent();
+		while (CurrentParent.IsValid())
+		{
+			FilterInNode(CurrentParent);
+
+			CurrentParent = CurrentParent->GetParent();
+			if (CurrentParent.IsValid() && CurrentParent->GetItemId() == FNavigationToolItemId::RootId)
+			{
+				CurrentParent = nullptr;
+			}
+		}
+	}
+
+	if (bInIncludeSelf)
+	{
+		FilterInNode(InNode);
+	}
+
+	if (bInIncludeChildren)
+	{
+		auto AddAllChildren = [this](const FNavigationToolItemPtr& InNode)
+		{
+			for (const FNavigationToolItemPtr& Child : InNode->GetChildren())
+			{
+				FilterInNode(Child);
+			}
+		};
+
+		for (const FNavigationToolItemPtr& Child : InNode->GetChildren())
+		{
+			AddAllChildren(Child);
+		}
+	}
+}
+
+void FNavigationToolFilterData::FilterInNodeWithAncestors(const FNavigationToolItemPtr& InNode)
+{
+	FilterInParentChildNodes(InNode, true, true, false);
+}
+
+bool FNavigationToolFilterData::ContainsFilterInNodes(const FNavigationToolFilterData& InOtherData) const
+{
+	return FilterInNodes.Includes(InOtherData.FilterInNodes);
+}
+
+bool FNavigationToolFilterData::IsFilteredIn(const FNavigationToolItemPtr& InNode) const
+{
+	return FilterInNodes.Contains(InNode);
+}
+
+bool FNavigationToolFilterData::IsFilteredOut(const FNavigationToolItemPtr& InNode) const
+{
+	return !FilterInNodes.Contains(InNode);
+}
+
+}
diff --git a/Engine/Plugins/Experimental/SequenceNavigator/Source/SequenceNavigator/Private/Filters/TextExpressions/NavigationToolFilterTextExpressionContext.cpp b/Engine/Plugins/Experimental/SequenceNavigator/Source/SequenceNavigator/Private/Filters/TextExpressions/NavigationToolFilterTextExpressionContext.cpp
new file mode 100644
--- /dev/null
+++ b/Engine/Plugins/Experimental/SequenceNavigator/Source/SequenceNavigator/Private/Filters/TextExpressions/NavigationToolFilterTextExpressionContext.cpp
@@ -0,0 +1,67 @@
+// Copyright Epic Games, Inc. All Rights Reserved.
+
+#include "Filters/TextExpressions/NavigationToolFilterTextExpressionContext.h"
+#include "Filters/INavigationToolFilterBar.h"
+#include "Items/INavigationToolItem.h"
+
+namespace UE::Sequencer::NavigationTool
+{
+
+using namespace UE::Sequencer;
+
+FNavigationToolFilterTextExpressionContext::FNavigationToolFilterTextExpressionContext(INavigationToolFilterBar& InFilterInterface)
+	: FilterInterface(InFilterInterface)
+{
+}
+
+void FNavigationToolFilterTextExpressionContext::SetFilterItem(const FNavigationToolItemPtr& InFilterItem)
+{
+	WeakFilterItem = InFilterItem;
+}
+
+bool FNavigationToolFilterTextExpressionContext::TestBasicStringExpression(const FTextFilterString& InValue
+	, const ETextFilterTextComparisonMode InTextComparisonMode) const
+{
+	if (const TSharedPtr<INavigationToolItem> FilterItem = WeakFilterItem.Pin())
+	{
+		const FString Label = FilterItem->GetDisplayName().ToString();
+		if (TextFilterUtils::TestBasicStringExpression(Label, InValue, ETextFilterTextComparisonMode::Partial))
+		{
+			return true;
+		}
+	}
+
+	return false;
+}
+
+bool FNavigationToolFilterTextExpressionContext::TestComplexExpression(const FName& InKey
+	, const FTextFilterString& InValue
+	, const ETextFilterComparisonOperation InComparisonOperation
+	, const ETextFilterTextComparisonMode InTextComparisonMode) const
+{
+	if (!WeakFilterItem.IsValid())
+	{
+		return false;
+	}
+
+	const TSet<FName> Keys = GetKeys();
+	if (!Keys.IsEmpty() && !Keys.Contains(InKey))
+	{
+		return false;
+	}
+
+	return !InValue.IsEmpty();
+}
+
+UMovieSceneSequence* FNavigationToolFilterTextExpressionContext::GetFocusedMovieSceneSequence() const
+{
+	return FilterInterface.GetSequencer().GetFocusedMovieSceneSequence();
+}
+
+UMovieScene* FNavigationToolFilterTextExpressionContext::GetFocusedGetMovieScene() const
+{
+	const UMovieSceneSequence* const FocusedMovieSceneSequence = GetFocusedMovieSceneSequence();
+	return IsValid(FocusedMovieSceneSequence) ? FocusedMovieSceneSequence->GetMovieScene() : nullptr;
+}
+
+}
diff --git a/Engine/Plugins/Experimental/SequenceNavigator/Source/SequenceNavigator/Private/ItemActions/NavigationToolAddItem.cpp b/Engine/Plugins/Experimental/SequenceNavigator/Source/SequenceNavigator/Private/ItemActions/NavigationToolAddItem.cpp
new file mode 100644
--- /dev/null
+++ b/Engine/Plugins/Experimental/SequenceNavigator/Source/SequenceNavigator/Private/ItemActions/NavigationToolAddItem.cpp
@@ -0,0 +1,123 @@
+// Copyright Epic Games, Inc. All Rights Reserved.
+
+#include "ItemActions/NavigationToolAddItem.h"
+#include "Items/INavigationToolItem.h"
+#include "NavigationTool.h"
+#include "NavigationToolScopedSelection.h"
+#include "Widgets/Views/STableRow.h"
+
+namespace UE::Sequencer::NavigationTool
+{
+
+FNavigationToolAddItem::FNavigationToolAddItem(const FNavigationToolAddItemParams& InAddItemParams)
+	: AddParams(InAddItemParams)
+{
+}
+
+bool FNavigationToolAddItem::ShouldTransact() const
+{
+	return EnumHasAnyFlags(AddParams.Flags, ENavigationToolAddItemFlags::Transact);
+}
+
+void FNavigationToolAddItem::Execute(FNavigationTool& InTool)
+{
+	if (!AddParams.Item.IsValid())
+	{
+		return;
+	}
+
+	// Try to Create Children on Find
+	if (EnumHasAnyFlags(AddParams.Flags, ENavigationToolAddItemFlags::AddChildren))
+	{
+		constexpr bool bRecursiveFind = true;
+		TArray<FNavigationToolItemPtr> Children;
+		AddParams.Item->FindValidChildren(Children, bRecursiveFind);
+	}
+
+	const FNavigationToolItemPtr ParentItem = AddParams.Item->GetParent();
+
+	// If this array has elements in it, then we need to stop a circular dependency from forming
+	const TArray<FNavigationToolItemPtr> PathToRelativeItem = AddParams.Item->FindPath({ AddParams.RelativeItem });
+	if (PathToRelativeItem.Num() > 0 && ParentItem.IsValid())
+	{
+		FNavigationToolAddItemParams CircularSolverParams;
+		CircularSolverParams.Item             = PathToRelativeItem[0];
+		CircularSolverParams.RelativeItem     = AddParams.Item;
+		CircularSolverParams.RelativeDropZone = EItemDropZone::AboveItem;
+		CircularSolverParams.Flags            = AddParams.Flags;
+
+		ParentItem->AddChild(MoveTemp(CircularSolverParams));
+	}
+
+	if (AddParams.RelativeItem.IsValid())
+	{
+		const FNavigationToolItemPtr RelativeItemParent = AddParams.RelativeItem->GetParent();
+
+		// If it's onto item, the Relative Item is going to be the Parent
+		if (!AddParams.RelativeDropZone.IsSet() || AddParams.RelativeDropZone == EItemDropZone::OntoItem)
+		{
+			//If the Relative Item is Onto and it's the same as the Current Parent, shift Item up in the Hierarchy
+			//(as long as the parent is a valid one)
+			if (AddParams.RelativeItem == ParentItem && RelativeItemParent)
+			{
+				AddParams.RelativeDropZone = EItemDropZone::BelowItem;
+				RelativeItemParent->AddChild(AddParams);
+			}
+			else
+			{
+				AddParams.RelativeItem->AddChild(AddParams);
+			}
+		}
+		//else we place it as a Sibling to the Relative Item
+		else if (RelativeItemParent)
+		{
+			RelativeItemParent->AddChild(AddParams);
+		}
+		//if no parent, then add it to the tree root
+		else
+		{
+			InTool.GetTreeRoot()->AddChild(AddParams);
+		}
+	}
+	else
+	{
+		//If no Relative Item, add to tree root
+		InTool.GetTreeRoot()->AddChild(AddParams);
+	}
+
+	const FNavigationToolScopedSelection ScopedSelection(*InTool.GetSequencer(), ENavigationToolScopedSelectionPurpose::Read);
+
+	// Automatically select item if it's selected
+	if (AddParams.Item->IsSelected(ScopedSelection))
+	{
+		// Select in Navigation Tool but don't signal selection as we already have it selected in Mode Tools
+		AddParams.Flags = ENavigationToolAddItemFlags::Select;
+		AddParams.SelectionFlags &= ~ENavigationToolItemSelectionFlags::SignalSelectionChange;
+	}
+	// Signal selection change when we attempt to select this item in the Navigation Tool but it isn't selected in Sequencer
+	else if (EnumHasAnyFlags(AddParams.Flags, ENavigationToolAddItemFlags::Select))
+	{
+		AddParams.SelectionFlags |= ENavigationToolItemSelectionFlags::SignalSelectionChange;
+	}
+
+	if (EnumHasAnyFlags(AddParams.Flags, ENavigationToolAddItemFlags::Select))
+	{
+		InTool.SelectItems({ AddParams.Item }, AddParams.SelectionFlags);
+	}
+
+	InTool.SetToolModified();
+}
+
+void FNavigationToolAddItem::OnObjectsReplaced(const TMap<UObject*, UObject*>& InReplacementMap, bool bRecursive)
+{
+	if (AddParams.Item.IsValid())
+	{
+		AddParams.Item->OnObjectsReplaced(InReplacementMap, bRecursive);
+	}
+	if (AddParams.RelativeItem.IsValid())
+	{
+		AddParams.RelativeItem->OnObjectsReplaced(InReplacementMap, bRecursive);
+	}
+}
+
+}
diff --git a/Engine/Plugins/Experimental/SequenceNavigator/Source/SequenceNavigator/Private/ItemActions/NavigationToolItemAction.cpp b/Engine/Plugins/Experimental/SequenceNavigator/Source/SequenceNavigator/Private/ItemActions/NavigationToolItemAction.cpp
new file mode 100644
--- /dev/null
+++ b/Engine/Plugins/Experimental/SequenceNavigator/Source/SequenceNavigator/Private/ItemActions/NavigationToolItemAction.cpp
@@ -0,0 +1,3 @@
+// Copyright Epic Games, Inc. All Rights Reserved.
+
+#include "ItemActions/NavigationToolItemAction.h"
diff --git a/Engine/Plugins/Experimental/SequenceNavigator/Source/SequenceNavigator/Private/ItemActions/NavigationToolRemoveItem.cpp b/Engine/Plugins/Experimental/SequenceNavigator/Source/SequenceNavigator/Private/ItemActions/NavigationToolRemoveItem.cpp
new file mode 100644
--- /dev/null
+++ b/Engine/Plugins/Experimental/SequenceNavigator/Source/SequenceNavigator/Private/ItemActions/NavigationToolRemoveItem.cpp
@@ -0,0 +1,95 @@
+// Copyright Epic Games, Inc. All Rights Reserved.
+
+#include "ItemActions/NavigationToolRemoveItem.h"
+#include "ItemActions/NavigationToolAddItem.h"
+#include "Items/INavigationToolItem.h"
+#include "NavigationTool.h"
+#include "Widgets/Views/STableRow.h"
+
+namespace UE::Sequencer::NavigationTool
+{
+
+FNavigationToolRemoveItem::FNavigationToolRemoveItem(const FNavigationToolRemoveItemParams& InRemoveItemParams)
+	: RemoveParams(InRemoveItemParams)
+{
+	if (RemoveParams.Item.IsValid())
+	{
+		RemoveParams.Item->AddFlags(ENavigationToolItemFlags::PendingRemoval);
+	}
+}
+
+void FNavigationToolRemoveItem::Execute(FNavigationTool& InTool)
+{
+	if (RemoveParams.Item.IsValid())
+	{
+		FNavigationToolItemFlagGuard Guard(RemoveParams.Item, ENavigationToolItemFlags::IgnorePendingKill);
+
+		//Copy the Array since we may modify it below on Add/Remove Child
+		TArray<FNavigationToolItemPtr> Children = RemoveParams.Item->GetChildren();
+
+		FNavigationToolItemPtr Parent       = RemoveParams.Item->GetParent();
+		FNavigationToolItemPtr RelativeItem = RemoveParams.Item;
+
+		//Search the lowest parent that is not pending removal
+		while (Parent.IsValid() && Parent->HasAnyFlags(ENavigationToolItemFlags::PendingRemoval))
+		{
+			RelativeItem = Parent;
+			Parent       = Parent->GetParent();
+		}
+
+		//Reparent the Item's Children to the Valid Parent found above
+		if (Parent.IsValid())
+		{
+			FNavigationToolAddItemParams ReparentParams;
+			ReparentParams.RelativeItem             = RelativeItem;
+			ReparentParams.RelativeDropZone         = EItemDropZone::BelowItem;
+			ReparentParams.Flags                    = ENavigationToolAddItemFlags::Select;
+
+			TArray<FNavigationToolItemPtr> ItemsToReparent;
+
+			ItemsToReparent.Append(Children);
+			while (ItemsToReparent.Num() > 0)
+			{
+				ReparentParams.Item = ItemsToReparent.Pop();
+
+				//Then Parent it to the Item's Parent.
+				//If we couldn't add the Child, then it means either the Child Item itself is invalid
+				//or the underlying info (e.g. object) is invalid (e.g. actor pending kill)
+				if (Parent->AddChild(ReparentParams) == false)
+				{
+					//so try reparenting the children of this invalid child (since this invalid child will be removed)
+					if (ReparentParams.Item.IsValid())
+					{
+						ItemsToReparent.Append(ReparentParams.Item->GetChildren());
+					}
+				}
+			}
+
+			//In case the Parent is still the same, Remove Child Item, else the Parent is going to be removed anyways
+			if (Parent == RemoveParams.Item->GetParent())
+			{
+				Parent->RemoveChild(RemoveParams.Item);
+			}
+		}
+		else
+		{
+			for (const FNavigationToolItemPtr& Child : Children)
+			{
+				RemoveParams.Item->RemoveChild(Child);
+			}
+		}
+
+		InTool.UnregisterItem(RemoveParams.Item->GetItemId());
+		InTool.SetToolModified();
+	}
+}
+
+void FNavigationToolRemoveItem::OnObjectsReplaced(const TMap<UObject*, UObject*>& InReplacementMap, bool bRecursive)
+{
+	if (RemoveParams.Item.IsValid())
+	{
+		RemoveParams.Item->OnObjectsReplaced(InReplacementMap, bRecursive);
+	}
+}
+
+}
diff --git a/Engine/Plugins/Experimental/SequenceNavigator/Source/SequenceNavigator/Private/Items/NavigationToolActor.cpp b/Engine/Plugins/Experimental/SequenceNavigator/Source/SequenceNavigator/Private/Items/NavigationToolActor.cpp
new file mode 100644
--- /dev/null
+++ b/Engine/Plugins/Experimental/SequenceNavigator/Source/SequenceNavigator/Private/Items/NavigationToolActor.cpp
@@ -0,0 +1,174 @@
+// Copyright Epic Games, Inc. All Rights Reserved.
+
+#include "Items/NavigationToolActor.h"
+#include "Editor/EditorEngine.h"
+#include "GameFramework/Actor.h"
+#include "ISequencer.h"
+#include "Items/NavigationToolComponent.h"
+#include "Items/NavigationToolSequence.h"
+#include "MovieScene.h"
+#include "MovieSceneCommonHelpers.h"
+#include "MovieSceneSequence.h"
+#include "NavigationTool.h"
+#include "NavigationToolSettings.h"
+#include "ScopedTransaction.h"
+#include "Utils/NavigationToolMiscUtils.h"
+
+#define LOCTEXT_NAMESPACE "NavigationToolActor"
+
+namespace UE::Sequencer::NavigationTool
+{
+
+FNavigationToolActor::FNavigationToolActor(INavigationTool& InTool
+	, const FNavigationToolItemPtr& InParentItem
+	, const TSharedPtr<FNavigationToolSequence>& InParentSequenceItem
+	, const FMovieSceneBinding& InBinding)
+	: Super(InTool, InParentItem, InParentSequenceItem, InBinding)
+{
+}
+
+void FNavigationToolActor::FindChildren(TArray<FNavigationToolItemPtr>& OutChildren, const bool bInRecursive)
+{
+	Super::FindChildren(OutChildren, bInRecursive);
+
+	const AActor* const ThisActor = GetActor();
+	if (!IsValid(ThisActor))
+	{
+		return;
+	}
+
+	const TSharedPtr<FNavigationToolSequence> ParentSequenceItem = WeakParentSequenceItem.Pin();
+	if (!ParentSequenceItem.IsValid())
+	{
+		return;
+	}
+
+	UMovieSceneSequence* const Sequence = ParentSequenceItem->GetSequence();
+	if (!IsValid(Sequence))
+	{
+		return;
+	}
+
+	UMovieScene* const MovieScene = Sequence->GetMovieScene();
+	if (!IsValid(MovieScene))
+	{
+		return;
+	}
+
+	const TSharedPtr<ISequencer> Sequencer = Tool.GetSequencer();
+	if (!Sequencer.IsValid())
+	{
+		return;
+	}
+
+	const TSharedRef<FNavigationToolActor> SharedThisRef = SharedThis(this);
+
+	const TArray<FMovieSceneBinding> Bindings = ParentSequenceItem->GetSortedBindings();
+	for (const FMovieSceneBinding& CurrentBinding : Bindings)
+	{
+		const FGuid& CurrentBindingGuid = CurrentBinding.GetObjectGuid();
+		const UClass* const CurrentObjectClass = MovieSceneHelpers::GetBoundObjectClass(Sequence, CurrentBindingGuid);
+		if (CurrentObjectClass && CurrentObjectClass->IsChildOf<UActorComponent>())
+		{
+			const TArrayView<TWeakObjectPtr<>> WeakBoundObjects = ResolveBoundObjects(*Sequencer, Sequence, CurrentBindingGuid);
+			const TWeakObjectPtr<> CurrentBoundObject = WeakBoundObjects.IsEmpty() ? nullptr : WeakBoundObjects[0];
+
+			const UActorComponent* const CurrentBoundActorComponent = Cast<UActorComponent>(CurrentBoundObject);
+			if (IsValid(CurrentBoundActorComponent) && CurrentBoundActorComponent->GetOwner() == ThisActor)
+			{
+				const FNavigationToolItemPtr NewItem = Tool.FindOrAdd<FNavigationToolComponent>(GetProvider()
+					, SharedThisRef, ParentSequenceItem, CurrentBinding);
+				OutChildren.Add(NewItem);
+				if (bInRecursive)
+				{
+					NewItem->FindChildren(OutChildren, bInRecursive);
+				}
+			}
+		}
+	}
+}
+
+bool FNavigationToolActor::IsAllowedInTool() const
+{
+	FNavigationTool& ToolPrivate = static_cast<FNavigationTool&>(Tool);
+
+	const AActor* const UnderlyingActor = GetActor();
+	if (!IsValid(UnderlyingActor))
+	{
+		// Always allow unbound binding items
+		return true;
+	}
+
+	const bool bActorAllowed = ToolPrivate.IsObjectAllowedInTool(UnderlyingActor);
+
+	return bActorAllowed;
+}
+
+ENavigationToolItemViewMode FNavigationToolActor::GetSupportedViewModes(const INavigationToolView& InToolView) const
+{
+	return ENavigationToolItemViewMode::ItemTree | ENavigationToolItemViewMode::HorizontalItemList;
+}
+
+TArray<FName> FNavigationToolActor::GetTags() const
+{
+	const AActor* const UnderlyingActor = GetActor();
+	if (IsValid(UnderlyingActor))
+	{
+		return UnderlyingActor->Tags;
+	}
+	return Super::GetTags();
+}
+
+bool FNavigationToolActor::GetVisibility() const
+{
+	const AActor* const UnderlyingActor = GetActor();
+	if (IsValid(UnderlyingActor))
+	{
+		return !UnderlyingActor->IsTemporarilyHiddenInEditor(true);
+	}
+	return false;
+}
+
+void FNavigationToolActor::OnVisibilityChanged(const bool bInNewVisibility)
+{
+	AActor* const UnderlyingActor = GetActor();
+	if (IsValid(UnderlyingActor))
+	{
+		UnderlyingActor->SetIsTemporarilyHiddenInEditor(!bInNewVisibility);
+	}
+}
+
+bool FNavigationToolActor::CanRename() const
+{
+	return Super::CanRename() && IsValid(GetActor());
+}
+
+bool FNavigationToolActor::Rename(const FString& InName)
+{
+	AActor* const UnderlyingActor = GetActor();
+	const bool bIsActorLabelEditable = IsValid(UnderlyingActor) && UnderlyingActor->IsActorLabelEditable();
+
+	bool bRenamed = false;
+	
+	if (bIsActorLabelEditable && !InName.Equals(UnderlyingActor->GetActorLabel(), ESearchCase::CaseSensitive))
+	{
+		const FScopedTransaction Transaction(LOCTEXT("NavigationToolRenameActor", "Rename Actor"));
+
+		FActorLabelUtilities::RenameExistingActor(UnderlyingActor, InName);
+
+		Super::Rename(InName);
+
+		bRenamed = true;
+	}
+
+	return bRenamed;
+}
+
+AActor* FNavigationToolActor::GetActor() const
+{
+	return WeakBoundObject.IsValid() ? Cast<AActor>(GetCachedBoundObject()) : nullptr;
+}
+
+}
+
+#undef LOCTEXT_NAMESPACE
diff --git a/Engine/Plugins/Experimental/SequenceNavigator/Source/SequenceNavigator/Private/Items/NavigationToolComponent.cpp b/Engine/Plugins/Experimental/SequenceNavigator/Source/SequenceNavigator/Private/Items/NavigationToolComponent.cpp
new file mode 100644
--- /dev/null
+++ b/Engine/Plugins/Experimental/SequenceNavigator/Source/SequenceNavigator/Private/Items/NavigationToolComponent.cpp
@@ -0,0 +1,141 @@
+// Copyright Epic Games, Inc. All Rights Reserved.
+
+#include "Items/NavigationToolComponent.h"
+#include "Components/PrimitiveComponent.h"
+#include "Components/SceneComponent.h"
+#include "Items/NavigationToolItemProxy.h"
+#include "NavigationTool.h"
+#include "ScopedTransaction.h"
+#include "Styling/StyleColors.h"
+#include "Widgets/Columns/SNavigationToolLabelComponent.h"
+
+#define LOCTEXT_NAMESPACE "NavigationToolComponent"
+
+namespace UE::Sequencer::NavigationTool
+{
+
+FNavigationToolComponent::FNavigationToolComponent(INavigationTool& InTool
+	, const FNavigationToolItemPtr& InParentItem
+	, const TSharedPtr<FNavigationToolSequence>& InParentSequenceItem
+	, const FMovieSceneBinding& InBinding)
+	: Super(InTool, InParentItem, InParentSequenceItem, InBinding)
+{
+}
+
+void FNavigationToolComponent::FindChildren(TArray<FNavigationToolItemPtr>& OutChildren, const bool bInRecursive)
+{
+	Super::FindChildren(OutChildren, bInRecursive);
+}
+
+void FNavigationToolComponent::GetItemProxies(TArray<TSharedPtr<FNavigationToolItemProxy>>& OutItemProxies)
+{
+	Super::GetItemProxies(OutItemProxies);
+
+	if (UPrimitiveComponent* const PrimitiveComponent = Cast<UPrimitiveComponent>(GetComponent()))
+	{
+		if (const TSharedPtr<FNavigationToolItemProxy> MaterialItemProxy = Tool.GetOrCreateItemProxy<FNavigationToolItemProxy>(SharedThis(this)))
+		{
+			OutItemProxies.Add(MaterialItemProxy);
+		}
+	}
+}
+
+bool FNavigationToolComponent::IsAllowedInTool() const
+{
+	FNavigationTool& ToolPrivate = static_cast<FNavigationTool&>(Tool);
+
+	const UActorComponent* const UnderlyingComponent = GetComponent();
+	if (!IsValid(UnderlyingComponent))
+	{
+		// Always allow unbound binding items
+		return true;
+	}
+
+	const bool bOwnerAllowed = ToolPrivate.IsObjectAllowedInTool(UnderlyingComponent->GetOwner());
+	const bool bComponentAllowed = ToolPrivate.IsObjectAllowedInTool(UnderlyingComponent);
+
+	return bOwnerAllowed && bComponentAllowed;
+}
+
+ENavigationToolItemViewMode FNavigationToolComponent::GetSupportedViewModes(const INavigationToolView& InToolView) const
+{
+	return ENavigationToolItemViewMode::ItemTree | ENavigationToolItemViewMode::HorizontalItemList;
+}
+
+TSharedRef<SWidget> FNavigationToolComponent::GenerateLabelWidget(const TSharedRef<SNavigationToolTreeRow>& InRow)
+{
+	return SNew(SNavigationToolLabelComponent, SharedThis(this), InRow);
+}
+
+FLinearColor FNavigationToolComponent::GetItemBrushColor() const
+{
+	return FStyleColors::White25.GetSpecifiedColor();
+}
+
+TArray<FName> FNavigationToolComponent::GetTags() const
+{
+	UActorComponent* const UnderlyingComponent = GetComponent();
+	if (IsValid(UnderlyingComponent))
+	{
+		return UnderlyingComponent->ComponentTags;
+	}
+	return Super::GetTags();
+}
+
+bool FNavigationToolComponent::GetVisibility() const
+{
+	const USceneComponent* const UnderlyingComponent = Cast<USceneComponent>(GetComponent());
+	if (IsValid(UnderlyingComponent))
+	{
+		return UnderlyingComponent->IsVisibleInEditor();
+	}
+	return false;
+}
+
+void FNavigationToolComponent::OnVisibilityChanged(const bool bInNewVisibility)
+{
+	USceneComponent* const UnderlyingComponent = Cast<USceneComponent>(GetComponent());
+	if (IsValid(UnderlyingComponent))
+	{
+		UnderlyingComponent->SetVisibility(bInNewVisibility);
+	}
+}
+
+bool FNavigationToolComponent::CanRename() const
+{
+	return Super::CanRename() && IsValid(GetComponent());
+}
+
+bool FNavigationToolComponent::Rename(const FString& InName)
+{
+	UActorComponent* const UnderlyingComponent = GetComponent();
+	if (!IsValid(UnderlyingComponent))
+	{
+		return false;
+	}
+
+	bool bRenamed = false;
+
+	if (!InName.Equals(UnderlyingComponent->GetName(), ESearchCase::CaseSensitive))
+	{
+		const FScopedTransaction Transaction(LOCTEXT("NavigationToolRenameComponent", "Rename Component"));
+
+		UnderlyingComponent->Modify();
+		UnderlyingComponent->Rename(*InName);
+
+		Super::Rename(InName);
+
+		bRenamed = true;
+	}
+
+	return bRenamed;
+}
+
+UActorComponent* FNavigationToolComponent::GetComponent() const
+{
+	return WeakBoundObject.IsValid() ? Cast<UActorComponent>(GetCachedBoundObject()) : nullptr;
+}
+
+}
+
+#undef LOCTEXT_NAMESPACE
diff --git a/Engine/Plugins/Experimental/SequenceNavigator/Source/SequenceNavigator/Private/Items/NavigationToolItem.cpp b/Engine/Plugins/Experimental/SequenceNavigator/Source/SequenceNavigator/Private/Items/NavigationToolItem.cpp
new file mode 100644
--- /dev/null
+++ b/Engine/Plugins/Experimental/SequenceNavigator/Source/SequenceNavigator/Private/Items/NavigationToolItem.cpp
@@ -0,0 +1,500 @@
+// Copyright Epic Games, Inc. All Rights Reserved.
+
+#include "Items/NavigationToolItem.h"
+#include "DragDropOps/NavigationToolItemDragDropOp.h"
+#include "Editor.h"
+#include "Framework/Application/SlateApplication.h"
+#include "ItemActions/NavigationToolAddItem.h"
+#include "ItemActions/NavigationToolRemoveItem.h"
+#include "Items/NavigationToolItemUtils.h"
+#include "Misc/Optional.h"
+#include "NavigationTool.h"
+#include "NavigationToolView.h"
+#include "Providers/NavigationToolProvider.h"
+#include "Styling/StyleColors.h"
+#include "Textures/SlateIcon.h"
+#include "Widgets/Columns/SNavigationToolLabelItem.h"
+#include "Widgets/Views/STableRow.h"
+
+#define LOCTEXT_NAMESPACE "NavigationToolItem"
+
+namespace UE::Sequencer::NavigationTool
+{
+
+FNavigationToolItem::FNavigationToolItem(INavigationTool& InTool, const FNavigationToolItemPtr& InParentItem)
+	: Tool(InTool)
+	, WeakParent(InParentItem)
+{
+}
+
+INavigationTool& FNavigationToolItem::GetOwnerTool() const
+{
+	return Tool;
+}
+
+TSharedPtr<FNavigationToolProvider> FNavigationToolItem::GetProvider() const
+{
+	return WeakProvider.Pin();
+}
+
+bool FNavigationToolItem::IsItemValid() const
+{
+	return true;
+}
+
+void FNavigationToolItem::RefreshChildren()
+{
+	TArray<FNavigationToolItemPtr> FoundChildren;
+	FindValidChildren(FoundChildren, /*bRecursiveFind=*/false);
+
+	TArray<FNavigationToolItemPtr> Sortable;
+	TArray<FNavigationToolItemPtr> Unsortable;
+	SplitSortableAndUnsortableItems(FoundChildren, Sortable, Unsortable);
+
+	// Start with all Sortable/Unsortable Items, and remove every item seen by iterating Children
+	TSet<FNavigationToolItemPtr> NewSortableChildren(Sortable);
+	TSet<FNavigationToolItemPtr> NewUnsortableChildren(Unsortable);
+	
+	// Remove items from "Children" that were not present in the Sortable Found Children (we'll add non-sortable later)
+	// Result is have Children only contain Items that existed previously
+	for (TArray<FNavigationToolItemPtr>::TIterator ItemIter(Children); ItemIter; ++ItemIter)
+	{
+		const FNavigationToolItemPtr Item(*ItemIter);
+
+		if (!Item.IsValid() || NewUnsortableChildren.Contains(Item))
+		{
+			ItemIter.RemoveCurrent();
+		}
+		else if (!NewSortableChildren.Contains(Item) || !Item->IsItemValid())
+		{
+			Item->SetParent(nullptr);
+			ItemIter.RemoveCurrent();
+		}
+
+		NewSortableChildren.Remove(Item);
+		NewUnsortableChildren.Remove(Item);
+	}
+
+	// Find Children for New Children in case these new children have grand children
+	// Note: This does not affect any of the current containers. It's just called for discovery
+	auto FindGrandChildren = [](const TSet<FNavigationToolItemPtr>& InChildren)
+		{
+			for (const FNavigationToolItemPtr& Child : InChildren)
+			{
+				TArray<FNavigationToolItemPtr> GrandChildren;
+				Child->FindValidChildren(GrandChildren, /*bRecursiveFind=*/true);
+			}
+		};
+
+	FindGrandChildren(NewUnsortableChildren);
+	FindGrandChildren(NewSortableChildren);
+
+	// After removing Children not present in Sortable
+	// Children should either be equal in size with Sortable (which means no new sortable children were added)
+	// or Sortable has more entries which means there are new items to add
+	if (Sortable.Num() > Children.Num())
+	{
+		check(!NewSortableChildren.IsEmpty());
+		HandleNewSortableChildren(NewSortableChildren.Array());
+	}
+
+	// Rearrange so that Children are arranged like so:
+	// [Unsortable Children][Sortable Children]
+	Unsortable.Append(MoveTemp(Children));
+	Children = MoveTemp(Unsortable);
+
+	// Update the Parents of every Child in the List
+	const FNavigationToolItemPtr This = SharedThis(this);
+	for (const FNavigationToolItemPtr& Child : Children)
+	{
+		Child->SetParent(This);
+	}
+}
+
+void FNavigationToolItem::ResetChildren()
+{
+	for (const FNavigationToolItemPtr& Item : GetChildren())
+	{
+		if (Item.IsValid())
+		{
+			Item->SetParent(nullptr);
+		}
+	}
+	GetChildrenMutable().Reset();	
+};
+
+void FNavigationToolItem::FindChildren(TArray<FNavigationToolItemPtr>& OutChildren, const bool bInRecursive)
+{
+	FNavigationTool& ToolPrivate = static_cast<FNavigationTool&>(Tool);
+
+	TArray<TSharedPtr<FNavigationToolItemProxy>> ItemProxies;
+	ToolPrivate.GetItemProxiesForItem(SharedThis(this), ItemProxies);
+	OutChildren.Reserve(OutChildren.Num() + ItemProxies.Num());
+
+	for (const TSharedPtr<FNavigationToolItemProxy>& ItemProxy : ItemProxies)
+	{
+		OutChildren.Add(ItemProxy);
+		if (bInRecursive)
+		{
+			ItemProxy->FindChildren(OutChildren, bInRecursive);
+		}
+	}
+}
+
+void FNavigationToolItem::FindValidChildren(TArray<FNavigationToolItemPtr>& OutChildren, bool bInRecursive)
+{
+	FindChildren(OutChildren, bInRecursive);
+
+	OutChildren.RemoveAll(
+		[](const FNavigationToolItemPtr& InItem)
+		{
+			return !InItem.IsValid() || !InItem->IsAllowedInTool();
+		});
+}
+
+TArray<FNavigationToolItemPtr> FNavigationToolItem::FindPath(const TArray<FNavigationToolItemPtr>& InItems) const
+{
+	const TSharedPtr<const INavigationToolItem> This = SharedThis(this);
+	TArray<FNavigationToolItemPtr> Path;
+	for (const FNavigationToolItemPtr& Item : InItems)
+	{
+		Path.Reset();
+		FNavigationToolItemPtr CurrentItem = Item;
+		while (CurrentItem.IsValid())
+		{
+			if (This == CurrentItem)
+			{
+				Algo::Reverse(Path);
+				return Path;
+			}
+			Path.Add(CurrentItem);
+			CurrentItem = CurrentItem->GetParent();
+		}
+	}
+	return TArray<FNavigationToolItemPtr>();
+}
+
+INavigationToolItem::IndexType FNavigationToolItem::GetChildIndex(const FNavigationToolItemPtr& ChildItem) const
+{
+	return GetChildren().Find(ChildItem);
+}
+
+bool FNavigationToolItem::CanAddChild(const FNavigationToolItemPtr& InChild) const
+{
+	return InChild.IsValid();
+}
+
+bool FNavigationToolItem::AddChild(const FNavigationToolAddItemParams& InAddItemParams)
+{
+	if (CanAddChild(InAddItemParams.Item))
+	{
+		AddChildChecked(InAddItemParams);
+		return true;
+	}
+	return false;
+}
+
+bool FNavigationToolItem::RemoveChild(const FNavigationToolRemoveItemParams& InRemoveItemParams)
+{
+	if (InRemoveItemParams.Item.IsValid())
+	{
+		return RemoveChildChecked(InRemoveItemParams);
+	}
+	return false;
+}
+
+void FNavigationToolItem::SetParent(FNavigationToolItemPtr InParent)
+{
+	//check that one of the parent's children is this
+	WeakParent = InParent;
+}
+
+ENavigationToolItemViewMode FNavigationToolItem::GetSupportedViewModes(const INavigationToolView& InToolView) const
+{
+	return InToolView.GetItemDefaultViewMode();
+}
+
+FNavigationToolItemId FNavigationToolItem::GetItemId() const
+{
+	if (ItemId.IsValidId())
+	{
+		return ItemId;
+	}
+	const_cast<FNavigationToolItem*>(this)->RecalculateItemId();
+	return ItemId;
+}
+
+const FSlateBrush* FNavigationToolItem::GetIconBrush() const
+{
+	const FSlateIcon Icon = FNavigationToolExtender::FindOverrideIcon(SharedThis(this));
+	if (Icon.IsSet())
+	{
+		return Icon.GetIcon();
+	}
+
+	if (const FSlateBrush* const DefaultIconBrush = GetDefaultIconBrush())
+	{
+		return DefaultIconBrush;
+	}
+
+	return GetIcon().GetIcon();
+}
+
+FSlateColor FNavigationToolItem::GetIconColor() const
+{
+	return FStyleColors::Foreground;
+}
+
+TSharedRef<SWidget> FNavigationToolItem::GenerateLabelWidget(const TSharedRef<SNavigationToolTreeRow>& InRow)
+{
+	return SNew(SNavigationToolLabelItem, SharedThis(this), InRow);
+}
+
+bool FNavigationToolItem::Delete()
+{
+	GetOwnerTool().NotifyToolItemDeleted(SharedThis(this));
+	return true;
+}
+
+void FNavigationToolItem::AddFlags(ENavigationToolItemFlags Flags)
+{
+	EnumAddFlags(ItemFlags, Flags);
+}
+
+void FNavigationToolItem::RemoveFlags(ENavigationToolItemFlags Flags)
+{
+	EnumRemoveFlags(ItemFlags, Flags);
+}
+
+bool FNavigationToolItem::HasAnyFlags(ENavigationToolItemFlags Flags) const
+{
+	return EnumHasAnyFlags(ItemFlags, Flags);
+}
+
+bool FNavigationToolItem::HasAllFlags(ENavigationToolItemFlags Flags) const
+{
+	return EnumHasAllFlags(ItemFlags, Flags);
+}
+
+TOptional<EItemDropZone> FNavigationToolItem::CanAcceptDrop(const FDragDropEvent& InDragDropEvent, EItemDropZone InDropZone)
+{
+	if (const TSharedPtr<FNavigationToolItemDragDropOp> ItemDragDropOp = InDragDropEvent.GetOperationAs<FNavigationToolItemDragDropOp>())
+	{
+		const TOptional<EItemDropZone> DropZone = ItemDragDropOp->CanDrop(InDropZone, SharedThis(this));
+		if (DropZone.IsSet())
+		{
+			if (GetItemId() != FNavigationToolItemId::RootId)
+			{
+				ItemDragDropOp->CurrentIconBrush = GetIconBrush();
+			}
+			return *DropZone;
+		}
+
+		ItemDragDropOp->CurrentIconBrush = FAppStyle::GetBrush(TEXT("Graph.ConnectorFeedback.Error"));
+	}
+
+	const TSharedRef<FNavigationToolItem> SharedThisRef = SharedThis(this);
+
+	TOptional<EItemDropZone> OutDropZone;
+
+	GetOwnerTool().ForEachProvider([this, &InDragDropEvent, InDropZone, &SharedThisRef, &OutDropZone]
+		(const TSharedRef<FNavigationToolProvider>& InProvider)
+		{
+			OutDropZone = InProvider->OnToolItemCanAcceptDrop(InDragDropEvent, InDropZone, SharedThisRef);
+			if (OutDropZone.IsSet())
+			{
+				return false;
+			}
+			return true;
+		});
+
+	return OutDropZone;
+}
+
+FReply FNavigationToolItem::AcceptDrop(const FDragDropEvent& InDragDropEvent, EItemDropZone InDropZone)
+{
+	if (const TSharedPtr<FNavigationToolItemDragDropOp> ItemDragDropOp = InDragDropEvent.GetOperationAs<FNavigationToolItemDragDropOp>())
+	{
+		const FReply Reply = ItemDragDropOp->Drop(InDropZone, SharedThis(this));
+		if (Reply.IsEventHandled())
+		{
+			return Reply;
+		}
+	}
+
+	const TSharedRef<FNavigationToolItem> SharedThisRef = SharedThis(this);
+
+	FReply OutReply = FReply::Unhandled();
+
+	GetOwnerTool().ForEachProvider([this, &InDragDropEvent, InDropZone, &SharedThisRef, &OutReply]
+		(const TSharedRef<FNavigationToolProvider>& InProvider)
+		{
+			OutReply = InProvider->OnToolItemAcceptDrop(InDragDropEvent, InDropZone, SharedThisRef);
+			if (OutReply.IsEventHandled())
+			{
+				return false;
+			}
+			return true;
+		});
+
+	return OutReply;
+}
+
+FLinearColor FNavigationToolItem::GetItemBrushColor() const
+{
+	return FStyleColors::White.GetSpecifiedColor();
+}
+
+bool FNavigationToolItem::IsExpanded() const
+{
+	if (const TSharedPtr<INavigationToolView> ToolView = Tool.GetMostRecentToolView())
+	{
+		return ToolView->IsItemExpanded(SharedThis(const_cast<FNavigationToolItem*>(this)));
+	}
+	return false;
+}
+
+void FNavigationToolItem::SetExpansion(const bool bInExpand)
+{
+	if (const TSharedPtr<INavigationToolView> ToolView = Tool.GetMostRecentToolView())
+	{
+		return ToolView->SetItemExpansion(SharedThis(this), bInExpand);
+	}
+}
+
+TOptional<FColor> FNavigationToolItem::GetColor() const
+{
+	return Tool.FindItemColor(SharedThis(const_cast<FNavigationToolItem*>(this)));
+}
+
+void FNavigationToolItem::SetColor(const TOptional<FColor>& InColor)
+{
+	Tool.SetItemColor(SharedThis(this), InColor.Get(FColor()));
+}
+
+void FNavigationToolItem::SetProvider(const TWeakPtr<FNavigationToolProvider>& InWeakProvider)
+{
+	WeakProvider = InWeakProvider;
+}
+
+void FNavigationToolItem::RecalculateItemId()
+{
+	const FNavigationToolItemId OldItemId = ItemId;
+	ItemId = CalculateItemId();
+
+	FNavigationTool& ToolPrivate = static_cast<FNavigationTool&>(Tool);
+	ToolPrivate.NotifyItemIdChanged(OldItemId, SharedThis(this));
+}
+
+void FNavigationToolItem::AddChildChecked(const FNavigationToolAddItemParams& InAddItemParams)
+{
+	if (const FNavigationToolItemPtr OldParent = InAddItemParams.Item->GetParent())
+	{
+		// if we are adding the child and the old parent is this,
+		// then it means we're just rearranging, only remove from array
+		if (OldParent.Get() == this)
+		{
+			Children.Remove(InAddItemParams.Item);
+		}
+		else
+		{
+			FNavigationToolRemoveItemParams RemoveParams(InAddItemParams.Item);
+			OldParent->RemoveChild(RemoveParams);
+		}
+	}
+
+	if (InAddItemParams.RelativeItem.IsValid()
+		&& InAddItemParams.RelativeDropZone.IsSet()
+		&& InAddItemParams.RelativeDropZone != EItemDropZone::OntoItem)
+	{
+		const int32 RelativeItemIndex = Children.Find(InAddItemParams.RelativeItem);
+		if (RelativeItemIndex != INDEX_NONE)
+		{
+			const int32 TargetIndex = InAddItemParams.RelativeDropZone == EItemDropZone::BelowItem
+				? RelativeItemIndex + 1
+				: RelativeItemIndex;
+
+			Children.EmplaceAt(TargetIndex, InAddItemParams.Item);
+		}
+		else
+		{
+			Children.EmplaceAt(0, InAddItemParams.Item);
+		}
+	}
+	else
+	{
+		Children.EmplaceAt(0, InAddItemParams.Item);
+	}
+
+	InAddItemParams.Item->SetParent(SharedThis(this));
+}
+
+bool FNavigationToolItem::RemoveChildChecked(const FNavigationToolRemoveItemParams& InRemoveItemParams)
+{
+	InRemoveItemParams.Item->SetParent(nullptr);
+	return Children.Remove(InRemoveItemParams.Item) > 0;
+}
+
+void FNavigationToolItem::HandleNewSortableChildren(TArray<FNavigationToolItemPtr> InSortableChildren)
+{
+	InSortableChildren.Sort([this](const FNavigationToolItemPtr& InItemA, const FNavigationToolItemPtr& InItemB)
+		{
+			const FNavigationToolSaveState* const SaveStateA = InItemA->GetProviderSaveState();
+			if (!SaveStateA)
+			{
+				return false;
+			}
+
+			const FNavigationToolSaveState* const SaveStateB = InItemB->GetProviderSaveState();
+			if (!SaveStateB)
+			{
+				return false;
+			}
+
+			const FNavigationToolSerializedTreeNode* const NodeA = SaveStateA->SerializedTree.FindTreeNode(InItemA->MakeSerializedItem());
+			const FNavigationToolSerializedTreeNode* const NodeB = SaveStateB->SerializedTree.FindTreeNode(InItemB->MakeSerializedItem());
+
+			return FNavigationToolSerializedTree::CompareTreeItemOrder(NodeA, NodeB);
+		});
+
+	FNavigationToolAddItemParams AddItemParams;
+	for (const FNavigationToolItemPtr& NewChild : InSortableChildren)
+	{
+		const FNavigationToolSaveState* const SaveState = NewChild->GetProviderSaveState();
+		if (!SaveState)
+		{
+			continue;
+		}
+
+		AddItemParams.Item = NewChild;
+
+		const FNavigationToolSerializedTreeNode* const TreeNode = SaveState->SerializedTree.FindTreeNode(NewChild->MakeSerializedItem());
+		if (TreeNode && Children.IsValidIndex(TreeNode->GetLocalIndex()))
+		{
+			// Add Before the Child at Index, so this Item is at the specific Index
+			AddItemParams.RelativeItem = Children[TreeNode->GetLocalIndex()];
+			AddItemParams.RelativeDropZone = EItemDropZone::AboveItem;
+		}
+		else
+		{
+			// Add After Last, so this Item is the last item in the List
+			AddItemParams.RelativeItem = Children.IsEmpty() ? nullptr : Children.Last();
+			AddItemParams.RelativeDropZone = EItemDropZone::BelowItem;
+		}
+
+		AddChild(AddItemParams);
+	}
+}
+
+FNavigationToolSaveState* FNavigationToolItem::GetProviderSaveState() const
+{
+	if (const TSharedPtr<FNavigationToolProvider> Provider = GetProvider())
+	{
+		return Provider->GetSaveState(Tool);
+	}
+	return nullptr;
+}
+
+}
+
+#undef LOCTEXT_NAMESPACE
diff --git a/Engine/Plugins/Experimental/SequenceNavigator/Source/SequenceNavigator/Private/Items/NavigationToolItemProxy.cpp b/Engine/Plugins/Experimental/SequenceNavigator/Source/SequenceNavigator/Private/Items/NavigationToolItemProxy.cpp
new file mode 100644
--- /dev/null
+++ b/Engine/Plugins/Experimental/SequenceNavigator/Source/SequenceNavigator/Private/Items/NavigationToolItemProxy.cpp
@@ -0,0 +1,69 @@
+// Copyright Epic Games, Inc. All Rights Reserved.
+
+#include "Items/NavigationToolItemProxy.h"
+#include "Input/Reply.h"
+#include "NavigationTool.h"
+#include "NavigationToolView.h"
+
+namespace UE::Sequencer::NavigationTool
+{
+
+FNavigationToolItemProxy::FNavigationToolItemProxy(INavigationTool& InTool, const FNavigationToolItemPtr& InParentItem)
+	: FNavigationToolItem(InTool, InParentItem)
+{
+	WeakParent = InParentItem;
+}
+
+bool FNavigationToolItemProxy::IsItemValid() const
+{
+	return FNavigationToolItem::IsItemValid()
+		&& WeakParent.IsValid()
+		&& Tool.FindItem(WeakParent.Pin()->GetItemId()).IsValid();
+}
+
+void FNavigationToolItemProxy::FindChildren(TArray<FNavigationToolItemPtr>& OutChildren, bool bRecursive)
+{
+	const FNavigationToolItemPtr Parent = GetParent();
+	if (!Parent.IsValid() || !Parent->IsAllowedInTool())
+	{
+		return;
+	}
+	
+	Super::FindChildren(OutChildren, bRecursive);
+	
+	GetProxiedItems(Parent.ToSharedRef(), OutChildren, bRecursive);
+}
+
+void FNavigationToolItemProxy::SetParent(FNavigationToolItemPtr InParent)
+{
+	Super::SetParent(InParent);
+
+	// Recalculate our item Id because we rely on what our parent is for our Id
+	RecalculateItemId();
+}
+
+ENavigationToolItemViewMode FNavigationToolItemProxy::GetSupportedViewModes(const INavigationToolView& InToolView) const
+{
+	// Hide proxies if it has no children
+	if (Children.IsEmpty())
+	{
+		return ENavigationToolItemViewMode::None;
+	}
+	return InToolView.GetItemProxyViewMode();
+}
+
+bool FNavigationToolItemProxy::CanAutoExpand() const
+{
+	return false;
+}
+
+FNavigationToolItemId FNavigationToolItemProxy::CalculateItemId() const
+{
+	if (WeakParent.IsValid())
+	{
+		return FNavigationToolItemId(WeakParent.Pin(), *this);	
+	}
+	return FNavigationToolItemId();
+}
+
+}
diff --git a/Engine/Plugins/Experimental/SequenceNavigator/Source/SequenceNavigator/Private/Items/NavigationToolItemUtils.cpp b/Engine/Plugins/Experimental/SequenceNavigator/Source/SequenceNavigator/Private/Items/NavigationToolItemUtils.cpp
new file mode 100644
--- /dev/null
+++ b/Engine/Plugins/Experimental/SequenceNavigator/Source/SequenceNavigator/Private/Items/NavigationToolItemUtils.cpp
@@ -0,0 +1,181 @@
+// Copyright Epic Games, Inc. All Rights Reserved.
+
+#include "Items/NavigationToolItemUtils.h"
+#include "Containers/Array.h"
+#include "Items/INavigationToolItem.h"
+#include "Items/NavigationToolSequence.h"
+#include "Items/NavigationToolTreeRoot.h"
+#include "MovieSceneSequence.h"
+#include "NavigationTool.h"
+#include "NavigationToolView.h"
+#include "ScopedTransaction.h"
+#include "Sections/MovieSceneCinematicShotSection.h"
+#include "Sections/MovieSceneSubSection.h"
+#include "TrackEditors/SubTrackEditorBase.h"
+
+#define LOCTEXT_NAMESPACE "NavigationToolItemUtils"
+
+namespace UE::Sequencer::NavigationTool
+{
+
+bool CompareToolItemOrder(const FNavigationToolItemPtr& InA, const FNavigationToolItemPtr& InB)
+{
+	if (!InA.IsValid())
+	{
+		return false;
+	}
+
+	if (!InB.IsValid())
+	{
+		return false;
+	}
+
+	if (const FNavigationToolItemPtr LowestCommonAncestor = FNavigationTool::FindLowestCommonAncestor({ InA, InB }))
+	{
+		const TArray<FNavigationToolItemPtr> PathToA = LowestCommonAncestor->FindPath({ InA });
+		const TArray<FNavigationToolItemPtr> PathToB = LowestCommonAncestor->FindPath({ InB });
+
+		int32 Index = 0;
+	
+		int32 PathAIndex = -1;
+		int32 PathBIndex = -1;
+
+		while (PathAIndex == PathBIndex)
+		{
+			if (!PathToA.IsValidIndex(Index))
+			{
+				return true;
+			}
+			if (!PathToB.IsValidIndex(Index))
+			{
+				return false;
+			}
+
+			PathAIndex = LowestCommonAncestor->GetChildIndex(PathToA[Index]);
+			PathBIndex = LowestCommonAncestor->GetChildIndex(PathToB[Index]);
+			++Index;
+		}
+
+		return PathAIndex < PathBIndex;
+	}
+
+	return false;
+}
+
+void SplitSortableAndUnsortableItems(const TArray<FNavigationToolItemPtr>& InItems
+	, TArray<FNavigationToolItemPtr>& OutSortable
+	, TArray<FNavigationToolItemPtr>& OutUnsortable)
+{
+	// Allocate both for worst case scenarios
+	OutSortable.Reserve(InItems.Num());
+	OutUnsortable.Reserve(InItems.Num());
+
+	for (const FNavigationToolItemPtr& Item : InItems)
+	{
+		if (!Item.IsValid())
+		{
+			continue;
+		}
+
+		if (Item->ShouldSort())
+		{
+			OutSortable.Add(Item);
+		}
+		else
+		{
+			OutUnsortable.Add(Item);
+		}
+	}
+}
+
+UMovieSceneSubSection* GetSequenceItemSubSection(const TSharedPtr<INavigationToolItem>& InItem)
+{
+	if (const FNavigationToolSequence* const SequenceItem = InItem->CastTo<FNavigationToolSequence>())
+	{
+		return SequenceItem->GetSubSection();
+	}
+	return nullptr;
+}
+
+UMovieSceneMetaData* GetSequenceItemMetaData(const TSharedPtr<INavigationToolItem>& InItem)
+{
+	if (const FNavigationToolSequence* const SequenceItem = InItem->CastTo<FNavigationToolSequence>())
+	{
+		return FSubTrackEditorUtil::FindOrAddMetaData(SequenceItem->GetSequence());
+	}
+	return nullptr;
+}
+
+void RemoveSequenceDisplayNameParentPrefix(FText& InOutDisplayName
+	, const TSharedRef<FNavigationToolSequence>& InSequenceItem)
+{
+	UMovieSceneSequence* const Sequence = InSequenceItem->GetSequence();
+	if (!IsValid(Sequence))
+	{
+		return;
+	}
+
+	const FNavigationToolItemPtr Parent = InSequenceItem->GetParent();
+	if (!Parent.IsValid() || Parent->IsA<FNavigationToolTreeRoot>())
+	{
+		return;
+	}
+
+	// Go up the hierarchy to find the parent sequence
+	const FNavigationToolSequence* const ParentSequenceItem = InSequenceItem->FindAncestorOfType<FNavigationToolSequence>();
+	if (!ParentSequenceItem)
+	{
+		return;
+	}
+
+	UMovieSceneSequence* const ParentSequence = ParentSequenceItem->GetSequence();
+	if (!IsValid(ParentSequence))
+	{
+		return;
+	}
+
+	const UMovieSceneCinematicShotSection* const ShotSection = Cast<UMovieSceneCinematicShotSection>(InSequenceItem->GetSubSection());
+
+	FText NewSequenceName = IsValid(ShotSection)
+		? FText::FromString(ShotSection->GetShotDisplayName()) : Sequence->GetDisplayName();
+	if (NewSequenceName.IsEmpty())
+	{
+		return;
+	}
+
+	// Remove the parent prefix from child name
+	static FString SeparatorString = TEXT("_");
+	const FString ParentSequenceNameStr = ParentSequence->GetDisplayName().ToString();
+
+	FString Prefix;
+	FString BaseName;
+	if (ParentSequenceNameStr.Split(SeparatorString, &Prefix, &BaseName))
+	{
+		Prefix += SeparatorString;
+
+		FString NewSequenceNameStr = NewSequenceName.ToString();
+
+		if (NewSequenceNameStr.StartsWith(Prefix)
+			&& NewSequenceNameStr.Len() > Prefix.Len())
+		{
+			NewSequenceNameStr.RemoveFromStart(Prefix);
+			NewSequenceName = FText::FromString(NewSequenceNameStr);
+		}
+	}
+
+	InOutDisplayName = NewSequenceName;
+}
+
+void AppendSequenceDisplayNameDirtyStatus(FText& InOutDisplayName
+	, const UMovieSceneSequence& InSequence)
+{
+	UPackage* const Package = InSequence.GetPackage();
+	if (IsValid(Package) && Package->IsDirty())
+	{
+		InOutDisplayName = FText::Format(LOCTEXT("DirtySymbol", "{0}*"), InOutDisplayName);
+	}
+}
+
+}
+
+#undef LOCTEXT_NAMESPACE
diff --git a/Engine/Plugins/Experimental/SequenceNavigator/Source/SequenceNavigator/Private/Items/NavigationToolSequence.cpp b/Engine/Plugins/Experimental/SequenceNavigator/Source/SequenceNavigator/Private/Items/NavigationToolSequence.cpp
new file mode 100644
--- /dev/null
+++ b/Engine/Plugins/Experimental/SequenceNavigator/Source/SequenceNavigator/Private/Items/NavigationToolSequence.cpp
@@ -0,0 +1,723 @@
+// Copyright Epic Games, Inc. All Rights Reserved.
+
+#include "Items/NavigationToolSequence.h"
+#include "GameFramework/Actor.h"
+#include "INavigationTool.h"
+#include "Items/NavigationToolActor.h"
+#include "Items/NavigationToolItemUtils.h"
+#include "Items/NavigationToolMarker.h"
+#include "Items/NavigationToolSubTrack.h"
+#include "Items/NavigationToolTrack.h"
+#include "LevelSequenceActor.h"
+#include "Misc/QualifiedFrameTime.h"
+#include "MovieScene.h"
+#include "MVVM/SectionModelStorageExtension.h"
+#include "MVVM/ViewModels/SectionModel.h"
+#include "NavigationToolScopedSelection.h"
+#include "ScopedTransaction.h"
+#include "Sections/MovieSceneCinematicShotSection.h"
+#include "Sections/MovieSceneSubSection.h"
+#include "SequencerChannelInterface.h"
+#include "SequencerSettings.h"
+#include "Styling/SlateIconFinder.h"
+#include "Tracks/MovieSceneCinematicShotTrack.h"
+#include "Utils/NavigationToolMiscUtils.h"
+#include "Utils/NavigationToolMovieSceneUtils.h"
+
+#define LOCTEXT_NAMESPACE "NavigationToolSequence"
+
+namespace UE::Sequencer::NavigationTool
+{
+
+FNavigationToolSequence::FNavigationToolSequence(INavigationTool& InTool
+	, const FNavigationToolItemPtr& InParentItem
+	, UMovieSceneSequence* const InSequence
+	, UMovieSceneSubSection* const InSubSection
+	, const int32 InSubSectionIndex)
+	: Super(InTool, InParentItem)
+	, IRevisionControlExtension(InSequence)
+	, WeakSubSection(InSubSection)
+	, SubSectionIndex(InSubSectionIndex)
+	, WeakSequence(InSequence)
+{
+}
+
+bool FNavigationToolSequence::IsItemValid() const
+{
+	return WeakSequence.IsValid();
+}
+
+UObject* FNavigationToolSequence::GetItemObject() const
+{
+	return GetSequence();
+}
+
+bool FNavigationToolSequence::IsAllowedInTool() const
+{
+	return WeakSequence.IsValid();
+}
+
+void FNavigationToolSequence::FindChildren(TArray<FNavigationToolItemPtr>& OutChildren, const bool bInRecursive)
+{
+	Super::FindChildren(OutChildren, bInRecursive);
+
+	UMovieSceneSequence* const Sequence = GetSequence();
+	if (!IsValid(Sequence))
+	{
+		return;
+	}
+
+	UMovieScene* const MovieScene = Sequence->GetMovieScene();
+	if (!IsValid(MovieScene))
+	{
+		return;
+	}
+
+	const TSharedRef<FNavigationToolSequence> SharedThisRef = SharedThis(this);
+
+	const TArray<UMovieSceneTrack*> MovieSceneTracks = MovieScene->GetTracks();
+	for (UMovieSceneTrack* const Track : MovieSceneTracks)
+	{
+		FNavigationToolItemPtr NewItem;
+
+		UMovieSceneSubTrack* const SubTrack = Cast<UMovieSceneSubTrack>(Track);
+		if (IsValid(SubTrack))
+		{
+			NewItem = Tool.FindOrAdd<FNavigationToolSubTrack>(GetProvider()
+				, SharedThisRef, SubTrack, Sequence, WeakSubSection, SubSectionIndex);
+		}
+		else if (IsValid(Track))
+		{
+			NewItem = Tool.FindOrAdd<FNavigationToolTrack>(GetProvider()
+				, SharedThisRef, Track, Sequence, WeakSubSection, SubSectionIndex);
+		}
+
+		if (NewItem.IsValid())
+		{
+			OutChildren.Add(NewItem);
+			if (bInRecursive)
+			{
+				NewItem->FindChildren(OutChildren, bInRecursive);
+			}
+		}
+	}
+
+	// Only show actor bindings for a Sequence
+	const TArray<FMovieSceneBinding> Bindings = GetSortedBindings();
+	for (const FMovieSceneBinding& Binding : Bindings)
+	{
+		const FGuid& BindingObjectGuid = Binding.GetObjectGuid();
+		const UClass* const BoundObjectClass = MovieSceneHelpers::GetBoundObjectClass(Sequence, BindingObjectGuid);
+		if (BoundObjectClass && BoundObjectClass->IsChildOf<AActor>())
+		{
+			const FNavigationToolItemPtr NewItem = Tool.FindOrAdd<FNavigationToolActor>(GetProvider()
+				, SharedThisRef, SharedThisRef, Binding);
+			OutChildren.Add(NewItem);
+			if (bInRecursive)
+			{
+				NewItem->FindChildren(OutChildren, bInRecursive);
+			}
+		}
+	}
+
+	// Add sequence marker items
+	const TArray<FMovieSceneMarkedFrame>& MarkedFrames = MovieScene->GetMarkedFrames();
+	for (int32 Index = 0; Index < MarkedFrames.Num(); ++Index)
+	{
+		const FNavigationToolItemPtr NewItem = Tool.FindOrAdd<FNavigationToolMarker>(GetProvider()
+			, SharedThisRef, SharedThisRef, Index);
+		OutChildren.Add(NewItem);
+		if (bInRecursive)
+		{
+			NewItem->FindChildren(OutChildren, bInRecursive);
+		}
+	}
+}
+
+void FNavigationToolSequence::GetItemProxies(TArray<TSharedPtr<FNavigationToolItemProxy>>& OutItemProxies)
+{
+	Super::GetItemProxies(OutItemProxies);
+}
+
+bool FNavigationToolSequence::AddChild(const FNavigationToolAddItemParams& InAddItemParams)
+{
+	// @TODO: handle sequence being moved to this item
+	return Super::AddChild(InAddItemParams);
+}
+
+bool FNavigationToolSequence::RemoveChild(const FNavigationToolRemoveItemParams& InRemoveItemParams)
+{
+	/// @TODO: handle sequence being moved to this item
+	return Super::RemoveChild(InRemoveItemParams);
+}
+
+ENavigationToolItemViewMode FNavigationToolSequence::GetSupportedViewModes(const INavigationToolView& InToolView) const
+{
+	// Sequences should only be visualized in Navigation Tool View and not appear in the Item Column List
+	// Support any other type of View Mode
+	return ENavigationToolItemViewMode::ItemTree | ~ENavigationToolItemViewMode::HorizontalItemList;
+}
+
+FText FNavigationToolSequence::GetDisplayName() const
+{
+	FText NewDisplayName = FText::GetEmpty();
+
+	const UMovieSceneCinematicShotSection* const ShotSection = Cast<UMovieSceneCinematicShotSection>(WeakSubSection.Get());
+	if (IsValid(ShotSection))
+	{
+		NewDisplayName = FText::FromString(ShotSection->GetShotDisplayName());
+	}
+	else
+	{
+		const UMovieSceneSequence* const Sequence = GetSequence();
+		if (IsValid(Sequence))
+		{
+			NewDisplayName = Sequence->GetDisplayName();
+		}
+	}
+
+	if (NewDisplayName.IsEmpty())
+	{
+		return NewDisplayName;
+	}
+
+	// Apply additional options to display name
+	const TSharedRef<FNavigationToolSequence> ThisSequenceRef = SharedThis(const_cast<FNavigationToolSequence*>(this));
+	const UNavigationToolSettings* const ToolSettings = GetDefault<UNavigationToolSettings>();
+
+	if (ToolSettings->ShouldUseShortNames())
+	{
+		RemoveSequenceDisplayNameParentPrefix(NewDisplayName, ThisSequenceRef);
+	}
+
+	UMovieSceneSequence* const Sequence = ThisSequenceRef->GetSequence();
+	if (IsValid(Sequence))
+	{
+		AppendSequenceDisplayNameDirtyStatus(NewDisplayName, *Sequence);
+	}
+
+	return NewDisplayName;
+}
+
+FText FNavigationToolSequence::GetClassName() const
+{
+	UMovieSceneSequence* const Sequence = GetSequence();
+	if (IsValid(Sequence))
+	{
+		return FText::FromString(Sequence->GetClass()->GetName());
+	}
+	return FText::FromString(UMovieSceneSequence::StaticClass()->GetName());
+}
+
+FSlateIcon FNavigationToolSequence::GetIcon() const
+{
+	return FSlateIconFinder::FindIconForClass(ALevelSequenceActor::StaticClass());
+}
+
+FText FNavigationToolSequence::GetIconTooltipText() const
+{
+	return ALevelSequenceActor::StaticClass()->GetDisplayNameText();
+}
+
+bool FNavigationToolSequence::IsSelected(const FNavigationToolScopedSelection& InSelection) const
+{
+	return WeakSubSection.IsValid() && InSelection.IsSelected(WeakSubSection.Get());
+}
+
+void FNavigationToolSequence::Select(FNavigationToolScopedSelection& InSelection) const
+{
+	if (WeakSubSection.IsValid())
+	{
+		InSelection.Select(WeakSubSection.Get());
+	}
+}
+
+void FNavigationToolSequence::OnSelect()
+{
+	Super::OnSelect();
+}
+
+void FNavigationToolSequence::OnDoubleClick()
+{
+	UMovieSceneSequence* const Sequence = GetSequence();
+	if (!IsValid(Sequence))
+	{
+		return;
+	}
+
+	const FModifierKeysState ModifierKeys = FSlateApplication::Get().GetModifierKeys();
+
+	if (ModifierKeys.IsAltDown())
+	{
+		// Instead of focusing the sequence in the Sequencer, focus the parent sequence and select sequence
+		if (const FNavigationToolSequence* const ParentSequenceItem = FindAncestorOfType<FNavigationToolSequence>())
+		{
+			UMovieSceneSequence* const ParentSequence = ParentSequenceItem->GetSequence();
+			if (IsValid(ParentSequence))
+			{
+				FocusSequence(Tool, *ParentSequence, *this);
+			}
+		}
+	}
+
+	FocusSequence(Tool, *Sequence);
+}
+
+void FNavigationToolSequence::OnObjectsReplaced(const TMap<UObject*, UObject*>& InReplacementMap, bool bInRecursive)
+{
+	// Get the Object even if it's Pending Kill (most likely it is)
+	const UMovieSceneSequence* const ObjectPendingKill = WeakSequence.Get(true);
+	if (ObjectPendingKill && InReplacementMap.Contains(ObjectPendingKill))
+	{
+		WeakSequence = Cast<UMovieSceneSequence>(InReplacementMap[ObjectPendingKill]);
+	}
+
+	// This handles calling OnObjectsReplaced for every child item
+	Super::OnObjectsReplaced(InReplacementMap, bInRecursive);
+}
+
+FNavigationToolItemId FNavigationToolSequence::CalculateItemId() const
+{
+	return FNavigationToolItemId(GetParent(), GetSequence(), WeakSubSection.Get(), SubSectionIndex);
+}
+
+bool FNavigationToolSequence::CanRename() const
+{
+	// Disable for now
+	return false;
+	/*const UMovieSceneSequence* const Sequence = GetSequence();
+	if (!IsValid(Sequence))
+	{
+		return false;
+	}
+
+	const UMovieScene* const MovieScene = Sequence->GetMovieScene();
+	if (!IsValid(MovieScene))
+	{
+		return false;
+	}
+
+	return !MovieScene->IsReadOnly();*/
+}
+
+bool FNavigationToolSequence::Rename(const FString& InName)
+{
+	UMovieSceneSequence* const Sequence = GetSequence();
+	if (!IsValid(Sequence))
+	{
+		return false;
+	}
+
+	bool bRenamed = false;
+
+	if (!InName.Equals(Sequence->GetName(), ESearchCase::CaseSensitive))
+	{
+		//const FScopedTransaction Transaction(LOCTEXT("SequenceRename", "Rename Sequence"));
+
+		if (UMovieSceneCinematicShotSection* const ShotSection = Cast<UMovieSceneCinematicShotSection>(WeakSubSection.Get()))
+		{
+			ShotSection->Modify();
+			ShotSection->SetShotDisplayName(InName);
+
+			Tool.NotifyToolItemRenamed(SharedThis(this));
+
+			bRenamed = true;
+		}
+		else
+		{
+			Sequence->Modify();
+
+			if (Sequence->Rename(*InName))
+			{
+				Tool.NotifyToolItemRenamed(SharedThis(this));
+
+				bRenamed = true;
+			}
+		}
+	}
+
+	return bRenamed;
+}
+
+EItemSequenceInactiveState FNavigationToolSequence::GetInactiveState() const
+{
+	if (WeakSubSection.IsValid())
+	{
+		return WeakSubSection->IsActive() ? EItemSequenceInactiveState::None : EItemSequenceInactiveState::Inactive;
+	}
+
+	const ENavigationToolCompareState State = CompareChildrenItemState<ISequenceInactivableExtension>(this,
+		[](const ISequenceInactivableExtension* const InItem)
+			{
+				return InItem->GetInactiveState() == EItemSequenceInactiveState::Inactive;
+			},
+		[](const ISequenceInactivableExtension* const InItem)
+			{
+				return InItem->GetInactiveState() == EItemSequenceInactiveState::None;
+			});
+
+	return static_cast<EItemSequenceInactiveState>(State);
+}
+
+void FNavigationToolSequence::SetIsInactive(const bool bInIsInactive)
+{
+	const bool bNewActiveState = !bInIsInactive;
+
+	if (WeakSubSection.IsValid())
+	{
+		if (!WeakSubSection->IsReadOnly() && WeakSubSection->IsActive() != bNewActiveState)
+		{
+			WeakSubSection->Modify();
+			WeakSubSection->SetIsActive(bNewActiveState);
+		}
+	}
+
+	for (ISequenceInactivableExtension* const InactivableItem : GetChildrenOfType<ISequenceInactivableExtension>())
+	{
+		InactivableItem->SetIsInactive(bInIsInactive);
+	}
+}
+
+EItemMarkerVisibility FNavigationToolSequence::GetMarkerVisibility() const
+{
+	return IsGloballyMarkedFramesForSequence(GetSequence())
+		? EItemMarkerVisibility::Visible : EItemMarkerVisibility::None;
+}
+
+void FNavigationToolSequence::SetMarkerVisibility(const bool bInVisible)
+{
+	const TSharedPtr<ISequencer> Sequencer = Tool.GetSequencer();
+	if (!Sequencer.IsValid())
+	{
+		return;
+	}
+
+	const bool bVisible = (GetMarkerVisibility() == EItemMarkerVisibility::Visible);
+	if (bVisible == bInVisible)
+	{
+		return;
+	}
+
+	UMovieSceneSequence* const Sequence = GetSequence();
+	if (Sequence == Sequencer->GetRootMovieSceneSequence())
+	{
+		Sequencer->GetSequencerSettings()->SetShowMarkedFrames(bInVisible);
+	}
+
+	ShowGloballyMarkedFramesForSequence(*Sequencer, Sequence, bInVisible);
+
+	for (IMarkerVisibilityExtension* const MarkerVisibilityItem : GetChildrenOfType<IMarkerVisibilityExtension>())
+	{
+		MarkerVisibilityItem->SetMarkerVisibility(bInVisible);
+	}
+}
+
+EItemSequenceLockState FNavigationToolSequence::GetLockState() const
+{
+	if (WeakSubSection.IsValid())
+	{
+		return WeakSubSection->IsLocked() ? EItemSequenceLockState::Locked : EItemSequenceLockState::None;
+	}
+
+	const ENavigationToolCompareState State = CompareChildrenItemState<ISequenceLockableExtension>(this,
+		[](const ISequenceLockableExtension* const InItem)
+			{
+				return InItem->GetLockState() == EItemSequenceLockState::Locked;
+			},
+		[](const ISequenceLockableExtension* const InItem)
+			{
+				return InItem->GetLockState() != EItemSequenceLockState::Locked;
+			});
+
+	return static_cast<EItemSequenceLockState>(State);
+}
+
+void FNavigationToolSequence::SetIsLocked(const bool bInIsLocked)
+{
+	if (WeakSubSection.IsValid())
+	{
+		if (WeakSubSection->IsLocked() != bInIsLocked)
+		{
+			WeakSubSection->Modify();
+			WeakSubSection->SetIsLocked(bInIsLocked);
+		}
+
+		return;
+	}
+
+	for (ISequenceLockableExtension* const LockableItem : GetChildrenOfType<ISequenceLockableExtension>())
+	{
+		LockableItem->SetIsLocked(bInIsLocked);
+	}
+}
+
+TOptional<FColor> FNavigationToolSequence::GetColor() const
+{
+	if (!WeakSubSection.IsValid())
+	{
+		return FColor::Transparent;
+	}
+
+	const TSharedPtr<ISequencer> Sequencer = Tool.GetSequencer();
+	if (!Sequencer.IsValid())
+	{
+		return FColor::Transparent;
+	}
+
+	const FColor TrackColor = WeakSubSection->GetColorTint();
+	return (TrackColor == FColor()) ? TOptional<FColor>() : WeakSubSection->GetColorTint();
+}
+
+void FNavigationToolSequence::SetColor(const TOptional<FColor>& InColor)
+{
+	if (!WeakSubSection.IsValid())
+	{
+		return;
+	}
+
+	const TSharedPtr<ISequencer> Sequencer = Tool.GetSequencer();
+	if (!Sequencer.IsValid())
+	{
+		return;
+	}
+
+	WeakSubSection->Modify();
+	WeakSubSection->SetColorTint(InColor.Get(FColor()));
+}
+
+EItemRevisionControlState FNavigationToolSequence::GetRevisionControlState() const
+{
+	UMovieSceneSequence* const Sequence = GetSequence();
+	if (!IsValid(Sequence))
+	{
+		return EItemRevisionControlState::None;
+	}
+
+	const FSourceControlStatePtr RevisionControlState = FindSourceControlState(Sequence->GetPackage());
+	if (!RevisionControlState.IsValid())
+	{
+		return EItemRevisionControlState::None;
+	}
+
+	return EItemRevisionControlState::SourceControlled;
+}
+
+const FSlateBrush* FNavigationToolSequence::GetRevisionControlStatusIcon() const
+{
+	UMovieSceneSequence* const Sequence = GetSequence();
+	if (IsValid(Sequence))
+	{
+		return FindSourceControlStatusBrush(Sequence->GetPackage());
+	}
+	return nullptr;
+}
+
+FText FNavigationToolSequence::GetRevisionControlStatusText() const
+{
+	UMovieSceneSequence* const Sequence = GetSequence();
+	if (IsValid(Sequence))
+	{
+		return FindSourceControlStatusText(Sequence->GetPackage());
+	}
+	return FText::GetEmpty();
+}
+
+EItemContainsPlayhead FNavigationToolSequence::ContainsPlayhead() const
+{
+	const TSharedPtr<ISequencer> Sequencer = Tool.GetSequencer();
+	if (!Sequencer.IsValid())
+	{
+		return EItemContainsPlayhead::None;
+	}
+
+	const FQualifiedFrameTime PlayheadTime = Sequencer->GetLocalTime();
+
+	UMovieSceneSubSection* const SubSection = GetSubSection();
+	if (IsValid(SubSection))
+	{
+		return SubSection->GetTrueRange().Contains(PlayheadTime.Time.FrameNumber)
+			? EItemContainsPlayhead::ContainsPlayhead : EItemContainsPlayhead::None;
+	}
+
+	UMovieSceneSequence* const RootSequence = Sequencer->GetRootMovieSceneSequence();
+	if (IsValid(RootSequence))
+	{
+		UMovieScene* const RootSequenceMovieScene = RootSequence->GetMovieScene();
+		if (IsValid(RootSequenceMovieScene))
+		{
+			return RootSequenceMovieScene->GetPlaybackRange().Contains(PlayheadTime.Time.FrameNumber)
+				? EItemContainsPlayhead::ContainsPlayhead : EItemContainsPlayhead::None;
+		}
+	}
+
+	return EItemContainsPlayhead::None;
+}
+
+FFrameNumber FNavigationToolSequence::GetInTime() const
+{
+	const UMovieSceneSequence* const Sequence = WeakSequence.Get();
+	if (!IsValid(Sequence))
+	{
+		return 0;
+	}
+
+	const UMovieSceneSubSection* const SubSection = WeakSubSection.Get();
+	if (!IsValid(SubSection))
+	{
+		return 0;
+	}
+
+	const FFrameTime LowerBoundValue = SubSection->SectionRange.Value.GetLowerBoundValue();
+	return LowerBoundValue.FrameNumber;
+}
+
+void FNavigationToolSequence::SetInTime(const FFrameNumber& InTime)
+{
+	const UMovieSceneSequence* const Sequence = WeakSequence.Get();
+	if (!IsValid(Sequence))
+	{
+		return;
+	}
+
+	UMovieSceneSubSection* const SubSection = WeakSubSection.Get();
+	if (!IsValid(SubSection))
+	{
+		return;
+	}
+
+	if (SubSection->SectionRange.Value.GetLowerBoundValue() == InTime)
+	{
+		return;
+	}
+
+	SubSection->Modify();
+	SubSection->SectionRange.Value.SetLowerBoundValue(InTime);
+}
+
+FFrameNumber FNavigationToolSequence::GetOutTime() const
+{
+	const UMovieSceneSequence* const Sequence = WeakSequence.Get();
+	if (!IsValid(Sequence))
+	{
+		return 0;
+	}
+
+	const UMovieSceneSubSection* const SubSection = WeakSubSection.Get();
+	if (!IsValid(SubSection))
+	{
+		return 0;
+	}
+
+	const FFrameTime UpperBoundValue = SubSection->SectionRange.Value.GetUpperBoundValue();
+	return UpperBoundValue.FrameNumber;
+}
+
+void FNavigationToolSequence::SetOutTime(const FFrameNumber& InTime)
+{
+	const UMovieSceneSequence* const Sequence = WeakSequence.Get();
+	if (!IsValid(Sequence))
+	{
+		return;
+	}
+
+	UMovieSceneSubSection* const SubSection = WeakSubSection.Get();
+	if (!IsValid(SubSection))
+	{
+		return;
+	}
+
+	if (SubSection->SectionRange.Value.GetUpperBoundValue() == InTime)
+	{
+		return;
+	}
+
+	SubSection->Modify();
+	SubSection->SectionRange.Value.SetUpperBoundValue(InTime);
+}
+
+UMovieSceneSequence* FNavigationToolSequence::GetSequence() const
+{
+	return WeakSequence.Get();
+}
+
+UMovieSceneSubSection* FNavigationToolSequence::GetSubSection() const
+{
+	return WeakSubSection.IsValid() ? WeakSubSection.Get() : nullptr;
+}
+
+int32 FNavigationToolSequence::GetSubSectionIndex() const
+{
+	return SubSectionIndex;
+}
+
+UMovieScene* FNavigationToolSequence::GetSequenceMovieScene() const
+{
+	UMovieSceneSequence* const Sequence = GetSequence();
+	if (IsValid(Sequence))
+	{
+		return Sequence->GetMovieScene();
+	}
+	return nullptr;
+}
+
+TArray<FMovieSceneBinding> FNavigationToolSequence::GetSortedBindings() const
+{
+	UMovieScene* const MovieScene = GetSequenceMovieScene();
+	if (!IsValid(MovieScene))
+	{
+		return {};
+	}
+
+	TArray<FMovieSceneBinding> Bindings = MovieScene->GetBindings();
+
+	Bindings.Sort([](const FMovieSceneBinding& InA, const FMovieSceneBinding& InB)
+		{
+			const int32 SortingOrderA = InA.GetSortingOrder();
+			const int32 SortingOrderB = InB.GetSortingOrder();
+			if (SortingOrderA == SortingOrderB)
+			{
+				return InA.GetName() < InB.GetName();
+			}
+			return SortingOrderA < SortingOrderB;
+		});
+
+	return MoveTemp(Bindings);
+}
+
+TViewModelPtr<FSectionModel> FNavigationToolSequence::GetViewModel() const
+{
+	if (!WeakSubSection.IsValid())
+	{
+		return nullptr;
+	}
+
+	const TSharedPtr<ISequencer> Sequencer = Tool.GetSequencer();
+	if (!Sequencer.IsValid())
+	{
+		return nullptr;
+	}
+
+	const TSharedPtr<FSequencerEditorViewModel> ViewModel = Sequencer->GetViewModel();
+	if (!ViewModel.IsValid())
+	{
+		return nullptr;
+	}
+
+	const FViewModelPtr RootViewModel = ViewModel->GetRootModel();
+	if (!RootViewModel.IsValid())
+	{
+		return nullptr;
+	}
+
+	const FSectionModelStorageExtension* const StorageExtension = RootViewModel->CastDynamic<FSectionModelStorageExtension>();
+	if (!StorageExtension)
+	{
+		return nullptr;
+	}
+
+	return StorageExtension->FindModelForSection(WeakSubSection.Get());
+}
+
+}
+
+#undef LOCTEXT_NAMESPACE
diff --git a/Engine/Plugins/Experimental/SequenceNavigator/Source/SequenceNavigator/Private/Items/NavigationToolSubTrack.cpp b/Engine/Plugins/Experimental/SequenceNavigator/Source/SequenceNavigator/Private/Items/NavigationToolSubTrack.cpp
new file mode 100644
--- /dev/null
+++ b/Engine/Plugins/Experimental/SequenceNavigator/Source/SequenceNavigator/Private/Items/NavigationToolSubTrack.cpp
@@ -0,0 +1,166 @@
+﻿// Copyright Epic Games, Inc. All Rights Reserved.
+
+#include "Items/NavigationToolSubTrack.h"
+#include "INavigationTool.h"
+#include "Items/NavigationToolItemUtils.h"
+#include "Items/NavigationToolSequence.h"
+#include "MovieSceneSequence.h"
+#include "MovieSceneTrack.h"
+#include "Sections/MovieSceneSubSection.h"
+#include "Tracks/MovieSceneSubTrack.h"
+#include "Utils/NavigationToolMovieSceneUtils.h"
+#include "Widgets/Views/STableRow.h"
+
+#define LOCTEXT_NAMESPACE "NavigationToolSubTrack"
+
+namespace UE::Sequencer::NavigationTool
+{
+
+using namespace UE::Sequencer;
+
+FNavigationToolSubTrack::FNavigationToolSubTrack(INavigationTool& InTool
+	, const FNavigationToolItemPtr& InParentItem
+	, UMovieSceneSubTrack* const InSubTrack
+	, const TWeakObjectPtr<UMovieSceneSequence>& InSequence
+	, const TWeakObjectPtr<UMovieSceneSection>& InSection
+	, const int32 InSubSectionIndex)
+	: Super(InTool
+	, InParentItem
+	, InSubTrack
+	, InSequence
+	, InSection
+	, InSubSectionIndex)
+{
+	OnTrackObjectChanged();
+}
+
+void FNavigationToolSubTrack::FindChildren(TArray<FNavigationToolItemPtr>& OutChildren, const bool bInRecursive)
+{
+	Super::FindChildren(OutChildren, bInRecursive);
+
+	UMovieSceneSubTrack* const SubTrack = GetSubTrack();
+	if (IsValid(SubTrack))
+	{
+		const TSharedRef<FNavigationToolSubTrack> SharedThisRef = SharedThis(this);
+
+		const TArray<UMovieSceneSection*>& AllSections = SubTrack->GetAllSections();
+		for (int32 Index = 0; Index < AllSections.Num(); ++Index)
+		{
+			UMovieSceneSubSection* const SubSection = Cast<UMovieSceneSubSection>(AllSections[Index]);
+			if (IsValid(SubSection))
+			{
+				UMovieSceneSequence* const Sequence = SubSection->GetSequence();
+				if (ensure(IsValid(Sequence)))
+				{
+					const FNavigationToolItemPtr NewItem = Tool.FindOrAdd<FNavigationToolSequence>(GetProvider(), SharedThisRef, Sequence, SubSection, Index);
+					OutChildren.Add(NewItem);
+					if (bInRecursive)
+					{
+						NewItem->FindChildren(OutChildren, bInRecursive);
+					}
+				}
+			}
+		}
+	}
+}
+
+TOptional<EItemDropZone> FNavigationToolSubTrack::CanAcceptDrop(const FDragDropEvent& InDragDropEvent, EItemDropZone InDropZone)
+{
+	return Super::CanAcceptDrop(InDragDropEvent, InDropZone);
+}
+
+FReply FNavigationToolSubTrack::AcceptDrop(const FDragDropEvent& InDragDropEvent, EItemDropZone InDropZone)
+{
+	return Super::AcceptDrop(InDragDropEvent, InDropZone);
+}
+
+UMovieSceneSubTrack* FNavigationToolSubTrack::GetSubTrack() const
+{
+	return Cast<UMovieSceneSubTrack>(GetTrack());
+}
+
+EItemSequenceInactiveState FNavigationToolSubTrack::GetInactiveState() const
+{
+	UMovieSceneSubSection* const ThisSubSection = Cast<UMovieSceneSubSection>(WeakSection.Get());
+	if (IsValid(ThisSubSection))
+	{
+		return ThisSubSection->IsActive() ? EItemSequenceInactiveState::None : EItemSequenceInactiveState::Inactive;
+	}
+
+	const ENavigationToolCompareState State = CompareChildrenItemState<ISequenceInactivableExtension>(this,
+		[](const ISequenceInactivableExtension* const InInactivableItem)
+			{
+				return InInactivableItem->GetInactiveState() == EItemSequenceInactiveState::Inactive;
+			},
+		[](const ISequenceInactivableExtension* const InInactivableItem)
+			{
+				return InInactivableItem->GetInactiveState() == EItemSequenceInactiveState::None;
+			});
+
+	return static_cast<EItemSequenceInactiveState>(State);
+}
+
+void FNavigationToolSubTrack::SetIsInactive(const bool bInIsInactive)
+{
+	const bool bNewActiveState = !bInIsInactive;
+
+	UMovieSceneSubSection* const ThisSubSection = Cast<UMovieSceneSubSection>(WeakSection.Get());
+	if (IsValid(ThisSubSection))
+	{
+		if (ThisSubSection->IsActive() != bNewActiveState)
+		{
+			ThisSubSection->Modify();
+			ThisSubSection->SetIsActive(bNewActiveState);
+		}
+	}
+
+	for (ISequenceInactivableExtension* const InactivableItem : GetChildrenOfType<ISequenceInactivableExtension>())
+	{
+		InactivableItem->SetIsInactive(bInIsInactive);
+	}
+}
+
+EItemMarkerVisibility FNavigationToolSubTrack::GetMarkerVisibility() const
+{
+	const ENavigationToolCompareState State = CompareChildrenItemState<IMarkerVisibilityExtension>(this,
+		[](const IMarkerVisibilityExtension* const InInactivableItem)
+			{
+				return InInactivableItem->GetMarkerVisibility() == EItemMarkerVisibility::Visible;
+			},
+		[](const IMarkerVisibilityExtension* const InInactivableItem)
+			{
+				return InInactivableItem->GetMarkerVisibility() == EItemMarkerVisibility::None;
+			});
+
+	return static_cast<EItemMarkerVisibility>(State);
+}
+
+void FNavigationToolSubTrack::SetMarkerVisibility(const bool bInVisible)
+{
+	const TSharedPtr<ISequencer> Sequencer = Tool.GetSequencer();
+	if (!Sequencer.IsValid())
+	{
+		return;
+	}
+
+	UMovieSceneSubSection* const ThisSubSection = Cast<UMovieSceneSubSection>(WeakSection.Get());
+	if (IsValid(ThisSubSection))
+	{
+		UMovieSceneSequence* const Sequence = ThisSubSection->GetSequence();
+		const bool bIsVisible = IsGloballyMarkedFramesForSequence(Sequence);
+		if (bIsVisible != bInVisible)
+		{
+			ModifySequenceAndMovieScene(Sequence);
+			ShowGloballyMarkedFramesForSequence(*Sequencer, Sequence, bInVisible);
+		}
+	}
+
+	for (IMarkerVisibilityExtension* const MarkerVisibilityItem : GetChildrenOfType<IMarkerVisibilityExtension>())
+	{
+		MarkerVisibilityItem->SetMarkerVisibility(bInVisible);
+	}
+}
+
+}
+
+#undef LOCTEXT_NAMESPACE
