[Sequencer] Move Navigation Tool to new experimental SequenceNavigator plugin

diff --git a/Engine/Plugins/Experimental/SequenceNavigator/Source/SequenceNavigator/Private/Items/NavigationToolTreeRoot.cpp b/Engine/Plugins/Experimental/SequenceNavigator/Source/SequenceNavigator/Private/Items/NavigationToolTreeRoot.cpp
new file mode 100644
--- /dev/null
+++ b/Engine/Plugins/Experimental/SequenceNavigator/Source/SequenceNavigator/Private/Items/NavigationToolTreeRoot.cpp
@@ -0,0 +1,94 @@
+// Copyright Epic Games, Inc. All Rights Reserved.
+
+#include "NavigationToolTreeRoot.h"
+#include "NavigationTool.h"
+#include "Textures/SlateIcon.h"
+#include "Widgets/SNullWidget.h"
+#include "Widgets/Views/STableRow.h"
+
+namespace UE::Sequencer::NavigationTool
+{
+
+FNavigationToolTreeRoot::FNavigationToolTreeRoot(INavigationTool& InTool)
+	: FNavigationToolItem(InTool, nullptr)
+{
+}
+
+void FNavigationToolTreeRoot::FindChildren(TArray<FNavigationToolItemPtr>& OutChildren, const bool bInRecursive)
+{
+	Super::FindChildren(OutChildren, bInRecursive);
+
+	const TSharedRef<FNavigationToolTreeRoot> SharedThisRef = SharedThis(this);
+
+	// Get items to add from providers
+	GetOwnerTool().ForEachProvider([this, &OutChildren, bInRecursive, &SharedThisRef]
+		(const TSharedRef<FNavigationToolProvider>& InToolProvider)
+		{
+			TArray<FNavigationToolItemPtr> ThisProviderItems;
+			InToolProvider->OnExtendItemChildren(Tool, SharedThisRef, ThisProviderItems, bInRecursive);
+
+			OutChildren.Append(ThisProviderItems);
+
+			return true;
+		});
+}
+
+bool FNavigationToolTreeRoot::CanAddChild(const FNavigationToolItemPtr& InChild) const
+{
+	return Super::CanAddChild(InChild) && InChild->CanBeTopLevel();
+}
+
+bool FNavigationToolTreeRoot::IsAllowedInTool() const
+{
+	checkNoEntry();
+	return false;
+}
+
+FText FNavigationToolTreeRoot::GetDisplayName() const
+{
+	checkNoEntry();
+	return FText::GetEmpty();
+}
+
+FText FNavigationToolTreeRoot::GetClassName() const
+{
+	checkNoEntry();
+	return FText::GetEmpty();
+}
+
+FText FNavigationToolTreeRoot::GetIconTooltipText() const
+{
+	checkNoEntry();
+	return FText::GetEmpty();
+}
+
+FSlateIcon FNavigationToolTreeRoot::GetIcon() const
+{
+	checkNoEntry();
+	return FSlateIcon();
+}
+
+TSharedRef<SWidget> FNavigationToolTreeRoot::GenerateLabelWidget(const TSharedRef<SNavigationToolTreeRow>& InRow)
+{
+	checkNoEntry();
+	return SNullWidget::NullWidget;
+}
+
+TOptional<EItemDropZone> FNavigationToolTreeRoot::CanAcceptDrop(const FDragDropEvent& InDragDropEvent, const EItemDropZone InDropZone)
+{
+	return Super::CanAcceptDrop(InDragDropEvent, InDropZone);
+}
+
+FReply FNavigationToolTreeRoot::AcceptDrop(const FDragDropEvent& InDragDropEvent, const EItemDropZone InDropZone)
+{
+	// Force tree root to always be OntoItem so we always create new items as children of root
+	constexpr EItemDropZone ForcedDropZone = EItemDropZone::OntoItem;
+	return Super::AcceptDrop(InDragDropEvent, ForcedDropZone);
+}
+
+FNavigationToolItemId FNavigationToolTreeRoot::CalculateItemId() const
+{
+	return FNavigationToolItemId::RootId;
+}
+
+}
diff --git a/Engine/Plugins/Experimental/SequenceNavigator/Source/SequenceNavigator/Private/Menus/NavigationToolViewMenuContext.h b/Engine/Plugins/Experimental/SequenceNavigator/Source/SequenceNavigator/Private/Menus/NavigationToolViewMenuContext.h
new file mode 100644
--- /dev/null
+++ b/Engine/Plugins/Experimental/SequenceNavigator/Source/SequenceNavigator/Private/Menus/NavigationToolViewMenuContext.h
@@ -0,0 +1,45 @@
+// Copyright Epic Games, Inc. All Rights Reserved.
+
+#pragma once
+
+#include "NavigationTool.h"
+#include "NavigationToolView.h"
+#include "Templates/SharedPointer.h"
+#include "UObject/Object.h"
+#include "ToolMenu.h"
+#include "NavigationToolViewMenuContext.generated.h"
+
+DECLARE_DELEGATE_OneParam(FOnPopulateNavigationToolViewToolbarMenu, UToolMenu*);
+
+UCLASS(Transient)
+class UNavigationToolViewMenuContext : public UObject
+{
+	GENERATED_BODY()
+
+	friend UE::Sequencer::NavigationTool::FNavigationToolView;
+
+public:
+	void Init(const TWeakPtr<UE::Sequencer::NavigationTool::FNavigationToolView>& InToolView)
+	{
+		WeakToolView = InToolView;
+	}
+
+	TSharedPtr<UE::Sequencer::NavigationTool::FNavigationToolView> GetToolView() const
+	{
+		return WeakToolView.Pin();
+	}
+
+	TSharedPtr<UE::Sequencer::NavigationTool::FNavigationTool> GetTool() const
+	{
+		if (const TSharedPtr<UE::Sequencer::NavigationTool::FNavigationToolView> ToolView = GetToolView())
+		{
+			return StaticCastSharedPtr<UE::Sequencer::NavigationTool::FNavigationTool>(ToolView->GetOwnerTool());
+		}
+		return nullptr;
+	}
+
+	FOnPopulateNavigationToolViewToolbarMenu OnPopulateMenu;
+
+private:
+	TWeakPtr<UE::Sequencer::NavigationTool::FNavigationToolView> WeakToolView;
+};
diff --git a/Engine/Plugins/Experimental/SequenceNavigator/Source/SequenceNavigator/Private/NavigationToolInstance.cpp b/Engine/Plugins/Experimental/SequenceNavigator/Source/SequenceNavigator/Private/NavigationToolInstance.cpp
new file mode 100644
--- /dev/null
+++ b/Engine/Plugins/Experimental/SequenceNavigator/Source/SequenceNavigator/Private/NavigationToolInstance.cpp
@@ -0,0 +1,8 @@
+﻿// Copyright Epic Games, Inc. All Rights Reserved.
+
+#include "NavigationToolInstance.h"
+
+namespace UE::Sequencer::NavigationTool
+{
+
+} // namespace UE::Sequencer::NavigationTool
diff --git a/Engine/Plugins/Experimental/SequenceNavigator/Source/SequenceNavigator/Private/NavigationToolInstance.h b/Engine/Plugins/Experimental/SequenceNavigator/Source/SequenceNavigator/Private/NavigationToolInstance.h
new file mode 100644
--- /dev/null
+++ b/Engine/Plugins/Experimental/SequenceNavigator/Source/SequenceNavigator/Private/NavigationToolInstance.h
@@ -0,0 +1,41 @@
+﻿// Copyright Epic Games, Inc. All Rights Reserved.
+
+#pragma once
+
+#include "Containers/Set.h"
+#include "Delegates/IDelegateInstance.h"
+#include "UObject/NameTypes.h"
+
+class ISequencer;
+
+namespace UE::Sequencer::NavigationTool
+{
+
+class FNavigationToolProvider;
+#if 1==2
+struct FNavigationToolInstance
+{
+	/** The Sequencer type name to register the provider for */
+	FName SequencerName;
+
+	/** The sequencer instance for this Navigation Tool instance */
+	TWeakPtr<ISequencer> WeakSequencer;
+
+	FDelegateHandle SequencerClosedHandle;
+
+	/** List of providers registered to this tool instance */
+	TSet<TSharedRef<FNavigationToolProvider>> Providers;
+
+	bool operator==(const FNavigationToolInstance& InInstance) const
+	{
+		return SequencerName == InInstance.SequencerName;
+	}
+
+	friend uint32 GetTypeHash(const FNavigationToolInstance& InInstance)
+	{
+		return GetTypeHash(InInstance.SequencerName);
+	}
+};
+#endif
+
+} // namespace UE::Sequencer::NavigationTool
diff --git a/Engine/Plugins/Experimental/SequenceNavigator/Source/SequenceNavigator/Private/NavigationToolView.cpp b/Engine/Plugins/Experimental/SequenceNavigator/Source/SequenceNavigator/Private/NavigationToolView.cpp
new file mode 100644
--- /dev/null
+++ b/Engine/Plugins/Experimental/SequenceNavigator/Source/SequenceNavigator/Private/NavigationToolView.cpp
@@ -0,0 +1,2218 @@
+// Copyright Epic Games, Inc. All Rights Reserved.
+
+#include "NavigationToolView.h"
+#include "Columns/NavigationToolColumn.h"
+#include "Columns/NavigationToolColumnExtender.h"
+#include "ContentBrowserModule.h"
+#include "DragDropOps/NavigationToolItemDragDropOp.h"
+#include "Filters/Filters/NavigationToolBuiltInFilter.h"
+#include "Filters/Filters/NavigationToolFilterBase.h"
+#include "Filters/NavigationToolFilterBar.h"
+#include "Framework/Application/SlateApplication.h"
+#include "Framework/Commands/GenericCommands.h"
+#include "Framework/MultiBox/MultiBoxBuilder.h"
+#include "IContentBrowserSingleton.h"
+#include "INavigationTool.h"
+#include "ISequencer.h"
+#include "Items/NavigationToolActor.h"
+#include "Items/NavigationToolComponent.h"
+#include "Items/NavigationToolSequence.h"
+#include "Items/NavigationToolTreeRoot.h"
+#include "LevelSequence.h"
+#include "Menus/NavigationToolItemContextMenu.h"
+#include "Misc/MessageDialog.h"
+#include "NavigationTool.h"
+#include "NavigationToolCommands.h"
+#include "NavigationToolExtender.h"
+#include "NavigationToolSettings.h"
+#include "Providers/NavigationToolProvider.h"
+#include "SequenceNavigatorLog.h"
+#include "Styling/SlateTypes.h"
+#include "Utils/NavigationToolMiscUtils.h"
+#include "Widgets/ModalTextInputDialog.h"
+#include "Widgets/SNavigationToolView.h"
+#include "Widgets/SNavigationToolTreeView.h"
+#include "Widgets/SNullWidget.h"
+
+#define LOCTEXT_NAMESPACE "NavigationToolView"
+
+namespace UE::Sequencer::NavigationTool
+{
+
+FNavigationToolView::FNavigationToolView(FPrivateToken)
+	: ToolSettings(GetMutableDefault<UNavigationToolSettings>())
+	, ItemContextMenu(MakeShared<FNavigationToolItemContextMenu>())
+{
+}
+
+FNavigationToolView::~FNavigationToolView()
+{
+	if (UObjectInitialized())
+	{
+		ToolSettings->OnSettingChanged().RemoveAll(this);
+	}
+}
+
+void FNavigationToolView::Init(const TSharedRef<FNavigationTool>& InTool, const TSharedPtr<FUICommandList>& InBaseCommandList)
+{
+	const TSharedPtr<ISequencer> Sequencer = InTool->GetSequencer();
+	if (!Sequencer.IsValid())
+	{
+		return;
+	}
+
+	WeakTool = InTool;
+
+	BindCommands(InBaseCommandList);
+
+	FilterBar = MakeShared<FNavigationToolFilterBar>(*InTool);
+	FilterBar->Init();
+	FilterBar->BindCommands(GetBaseCommandList());
+	FilterBar->OnStateChanged().AddSPLambda(this,
+		[this](const bool bInIsVisible, const EFilterBarLayout InNewLayout)
+		{
+			if (const TSharedPtr<SNavigationToolView> ToolWidget = StaticCastSharedPtr<SNavigationToolView>(GetToolWidget()))
+			{
+				ToolWidget->RebuildWidget();
+			}
+		});
+
+	ToolSettings->OnSettingChanged().AddSP(this, &FNavigationToolView::OnToolSettingsChanged);
+
+	ToolViewWidget = SNew(SNavigationToolView, SharedThis(this));
+
+	UpdateRecentViews();
+}
+
+void FNavigationToolView::CreateColumns(const TSharedRef<FNavigationToolProvider>& InProvider)
+{
+	const TSharedPtr<FNavigationTool> Tool = WeakTool.Pin();
+	if (!Tool.IsValid())
+	{
+		return;
+	}
+
+	FNavigationToolColumnExtender ColumnExtender;
+	InProvider->OnExtendColumns(ColumnExtender);
+
+	// Sort and re-cache columns
+	const TArray<TSharedPtr<FNavigationToolColumn>>& ColumnsToAdd = ColumnExtender.GetColumns();
+
+	for (const TSharedPtr<FNavigationToolColumn>& Column : ColumnsToAdd)
+	{
+		const FName ColumnId = Column->GetColumnId();
+		if (!Columns.Contains(ColumnId))
+		{
+			Columns.Add(ColumnId, Column);
+		}
+	}
+
+	if (ToolViewWidget.IsValid())
+	{
+		ToolViewWidget->ReconstructColumns();
+	}
+}
+
+void FNavigationToolView::CreateDefaultColumnViews(const TSharedRef<FNavigationToolProvider>& InProvider)
+{
+	InProvider->OnExtendColumnViews(ToolSettings->GetCustomColumnViews());
+	ToolSettings->SaveConfig();
+}
+
+TSharedRef<FNavigationToolView> FNavigationToolView::CreateInstance(const int32 InToolViewId
+	, const TSharedRef<FNavigationTool>& InTool
+	, const TSharedPtr<FUICommandList>& InBaseCommandList)
+{
+	const TSharedRef<FNavigationToolView> Instance = MakeShared<FNavigationToolView>(FPrivateToken{});
+	Instance->ToolViewId = InToolViewId;
+	Instance->Init(InTool, InBaseCommandList);
+	return Instance;
+}
+
+void FNavigationToolView::PostLoad()
+{
+	if (ToolViewWidget.IsValid())
+	{
+		ToolViewWidget->ReconstructColumns();
+	}
+}
+
+void FNavigationToolView::OnToolSettingsChanged(UObject* InObject, FPropertyChangedEvent& InPropertyChangedEvent)
+{
+	RefreshTool(false);
+}
+
+void FNavigationToolView::Tick(float InDeltaTime)
+{
+	if (bRefreshRequested)
+	{
+		bRefreshRequested = false;
+		Refresh();
+	}
+
+	for (const TPair<FName, TSharedPtr<INavigationToolColumn>>& Pair : Columns)
+	{
+		if (const TSharedPtr<INavigationToolColumn>& Column = Pair.Value)
+		{
+			Column->Tick(InDeltaTime);
+		}
+	}
+
+	// Check if we have pending items to rename and we are not currently renaming an item
+	if (bRenamingItems && ItemsRemainingRename.Num() > 0 && !CurrentItemRenaming.IsValid())
+	{
+		CurrentItemRenaming = ItemsRemainingRename[0];
+		ItemsRemainingRename.RemoveAt(0);
+
+		if (CurrentItemRenaming.IsValid())
+		{
+			CurrentItemRenaming->OnRenameAction().AddSP(this, &FNavigationToolView::OnItemRenameAction);
+			CurrentItemRenaming->OnRenameAction().Broadcast(ENavigationToolRenameAction::Requested, SharedThis(this));
+		}
+	}
+
+	if (bRequestedRename)
+	{
+		bRequestedRename = false;
+		RenameSelected();
+	}
+}
+
+void FNavigationToolView::BindCommands(const TSharedPtr<FUICommandList>& InBaseCommandList)
+{
+	const FGenericCommands& GenericCommands = FGenericCommands::Get();
+	const FNavigationToolCommands& ToolCommands = FNavigationToolCommands::Get();
+
+	ViewCommandList = MakeShared<FUICommandList>();
+
+	if (InBaseCommandList.IsValid())
+	{
+		InBaseCommandList->Append(ViewCommandList.ToSharedRef());
+	}
+
+	ViewCommandList->MapAction(ToolCommands.OpenToolSettings
+		, FExecuteAction::CreateStatic(&UNavigationToolSettings::OpenEditorSettings));
+
+	ViewCommandList->MapAction(ToolCommands.Refresh
+		, FExecuteAction::CreateSP(this, &FNavigationToolView::RefreshTool, true));
+
+	ViewCommandList->MapAction(GenericCommands.Rename
+		, FExecuteAction::CreateSP(this, &FNavigationToolView::RenameSelected)
+		, FCanExecuteAction::CreateSP(this, &FNavigationToolView::CanRenameSelected));
+
+	ViewCommandList->MapAction(GenericCommands.Delete
+		, FExecuteAction::CreateSP(this, &FNavigationToolView::DeleteSelected)
+		, FCanExecuteAction::CreateSP(this, &FNavigationToolView::CanDeleteSelected));
+
+	ViewCommandList->MapAction(GenericCommands.Duplicate
+		, FExecuteAction::CreateSP(this, &FNavigationToolView::DuplicateSelected)
+		, FCanExecuteAction::CreateSP(this, &FNavigationToolView::CanDuplicateSelected));
+
+	ViewCommandList->MapAction(ToolCommands.SelectAllChildren
+		, FExecuteAction::CreateSP(this, &FNavigationToolView::SelectChildren, true)
+		, FCanExecuteAction::CreateSP(this, &FNavigationToolView::CanSelectChildren));
+
+	ViewCommandList->MapAction(ToolCommands.SelectImmediateChildren
+		, FExecuteAction::CreateSP(this, &FNavigationToolView::SelectChildren, false)
+		, FCanExecuteAction::CreateSP(this, &FNavigationToolView::CanSelectChildren));
+
+	ViewCommandList->MapAction(ToolCommands.SelectParent
+		, FExecuteAction::CreateSP(this, &FNavigationToolView::SelectParent)
+		, FCanExecuteAction::CreateSP(this, &FNavigationToolView::CanSelectParent));
+
+	ViewCommandList->MapAction(ToolCommands.SelectFirstChild
+		, FExecuteAction::CreateSP(this, &FNavigationToolView::SelectFirstChild)
+		, FCanExecuteAction::CreateSP(this, &FNavigationToolView::CanSelectFirstChild));
+
+	ViewCommandList->MapAction(ToolCommands.SelectNextSibling
+		, FExecuteAction::CreateSP(this, &FNavigationToolView::SelectSibling, +1)
+		, FCanExecuteAction::CreateSP(this, &FNavigationToolView::CanSelectSibling));
+
+	ViewCommandList->MapAction(ToolCommands.SelectPreviousSibling
+		, FExecuteAction::CreateSP(this, &FNavigationToolView::SelectSibling, -1)
+		, FCanExecuteAction::CreateSP(this, &FNavigationToolView::CanSelectSibling));
+
+	ViewCommandList->MapAction(ToolCommands.ExpandAll
+		, FExecuteAction::CreateSP(this, &FNavigationToolView::ExpandAll)
+		, FCanExecuteAction::CreateSP(this, &FNavigationToolView::CanExpandAll));
+
+	ViewCommandList->MapAction(ToolCommands.CollapseAll
+		, FExecuteAction::CreateSP(this, &FNavigationToolView::CollapseAll)
+		, FCanExecuteAction::CreateSP(this, &FNavigationToolView::CanCollapseAll));
+
+	ViewCommandList->MapAction(ToolCommands.ScrollNextSelectionIntoView
+		, FExecuteAction::CreateSP(this, &FNavigationToolView::ScrollNextIntoView)
+		, FCanExecuteAction::CreateSP(this, &FNavigationToolView::CanScrollNextIntoView));
+
+	ViewCommandList->MapAction(ToolCommands.ToggleMutedHierarchy
+		, FExecuteAction::CreateSP(this, &FNavigationToolView::ToggleMutedHierarchy)
+		, FCanExecuteAction::CreateSP(this, &FNavigationToolView::CanToggleMutedHierarchy)
+		, FIsActionChecked::CreateSP(this, &FNavigationToolView::IsMutedHierarchyActive));
+
+	ViewCommandList->MapAction(ToolCommands.ToggleAutoExpandToSelection
+		, FExecuteAction::CreateSP(this, &FNavigationToolView::ToggleAutoExpandToSelection)
+		, FCanExecuteAction::CreateSP(this, &FNavigationToolView::CanToggleAutoExpandToSelection)
+		, FIsActionChecked::CreateSP(this, &FNavigationToolView::ShouldAutoExpandToSelection));
+
+	ViewCommandList->MapAction(ToolCommands.ToggleShortNames
+		, FExecuteAction::CreateSP(this, &FNavigationToolView::ToggleUseShortNames)
+		, FCanExecuteAction::CreateSP(this, &FNavigationToolView::CanToggleUseShortNames)
+		, FIsActionChecked::CreateSP(this, &FNavigationToolView::ShouldUseShortNames));
+
+	ViewCommandList->MapAction(ToolCommands.ResetVisibleColumnSizes
+		, FExecuteAction::CreateSP(this, &FNavigationToolView::ResetVisibleColumnSizes)
+		, FCanExecuteAction::CreateSP(this, &FNavigationToolView::CanResetAllColumnSizes));
+
+	ViewCommandList->MapAction(ToolCommands.SaveCurrentColumnView
+		, FExecuteAction::CreateSP(this, &FNavigationToolView::SaveNewCustomColumnView));
+
+	ViewCommandList->MapAction(ToolCommands.FocusSingleSelection
+		, FExecuteAction::CreateSP(this, &FNavigationToolView::FocusSingleSelection)
+		, FCanExecuteAction::CreateSP(this, &FNavigationToolView::CanFocusSingleSelection));
+
+	ViewCommandList->MapAction(ToolCommands.FocusInContentBrowser
+		, FExecuteAction::CreateSP(this, &FNavigationToolView::FocusInContentBrowser)
+		, FCanExecuteAction::CreateSP(this, &FNavigationToolView::CanFocusInContentBrowser));
+}
+
+TSharedPtr<FUICommandList> FNavigationToolView::GetBaseCommandList() const
+{
+	if (const TSharedPtr<FNavigationTool> Tool = WeakTool.Pin())
+	{
+		return Tool->GetBaseCommandList();
+	}
+	return nullptr;
+}
+
+void FNavigationToolView::UpdateRecentViews()
+{
+	if (const TSharedPtr<FNavigationTool> Tool = WeakTool.Pin())
+	{
+		Tool->UpdateRecentToolViews(ToolViewId);
+	}
+}
+
+bool FNavigationToolView::IsMostRecentToolView() const
+{
+	return WeakTool.IsValid() && WeakTool.Pin()->GetMostRecentToolView().Get() == this;
+}
+
+TSharedPtr<ISequencer> FNavigationToolView::GetSequencer() const
+{
+	if (const TSharedPtr<INavigationTool> Tool = GetOwnerTool())
+	{
+		return Tool->GetSequencer();
+	}
+	return nullptr;
+}
+
+TSharedPtr<INavigationTool> FNavigationToolView::GetOwnerTool() const
+{
+	return WeakTool.Pin();
+}
+
+TSharedPtr<SWidget> FNavigationToolView::GetToolWidget() const
+{
+	return ToolViewWidget;
+}
+
+TSharedPtr<SWidget> FNavigationToolView::CreateItemContextMenu()
+{
+	return ItemContextMenu->CreateMenu(SharedThis(this), SelectedItems);
+}
+
+bool FNavigationToolView::ShouldShowColumnByDefault(const TSharedPtr<INavigationToolColumn>& InColumn) const
+{
+	if (!InColumn.IsValid())
+	{
+		return false;
+	}
+
+	const TSharedPtr<FNavigationTool> Tool = WeakTool.Pin();
+	if (!Tool.IsValid())
+	{
+		return false;
+	}
+
+	bool bShouldShow = IsColumnVisible(InColumn);
+
+	bShouldShow |= InColumn->ShouldShowColumnByDefault();
+
+	return bShouldShow;
+}
+
+void FNavigationToolView::RequestRefresh()
+{
+	bRefreshRequested = true;
+}
+
+void FNavigationToolView::Refresh()
+{
+	UpdateRootVisibleItems();
+
+	bFilterUpdateRequested = true;
+	UpdateFilters();
+
+	UpdateItemExpansions();
+	
+	if (ToolViewWidget.IsValid())
+	{
+		ToolViewWidget->RequestTreeRefresh();
+	}
+	
+	OnToolViewRefreshed.Broadcast();
+}
+
+void FNavigationToolView::SetKeyboardFocus()
+{
+	if (ToolViewWidget.IsValid())
+	{
+		ToolViewWidget->SetKeyboardFocus();
+	}
+}
+
+void FNavigationToolView::UpdateRootVisibleItems()
+{
+	RootVisibleItems.Reset();
+	ReadOnlyItems.Reset();
+	
+	if (WeakTool.IsValid())
+	{
+		GetChildrenOfItem(WeakTool.Pin()->GetTreeRoot(), RootVisibleItems);
+	}
+}
+
+void FNavigationToolView::UpdateItemExpansions()
+{
+	TArray<FNavigationToolItemPtr> Items = RootVisibleItems;
+
+	while (Items.Num() > 0)
+	{
+		FNavigationToolItemPtr Item = Items.Pop();
+		const ENavigationToolItemFlags ItemFlags = GetViewItemFlags(Item);
+		SetItemExpansion(Item, EnumHasAnyFlags(ItemFlags, ENavigationToolItemFlags::Expanded));
+		Items.Append(Item->GetChildren());
+	}
+
+	Items = RootVisibleItems;
+
+	while (Items.Num() > 0)
+	{
+		const FNavigationToolItemPtr Item = Items.Pop();
+		if (ToolViewWidget.IsValid())
+		{
+			ToolViewWidget->UpdateItemExpansions(Item);
+		}
+		Items.Append(Item->GetChildren());
+	}
+}
+
+void FNavigationToolView::NotifyObjectsReplaced()
+{
+	if (ToolViewWidget.IsValid())
+	{
+		ToolViewWidget->Invalidate(EInvalidateWidgetReason::Paint);	
+	}
+}
+
+FNavigationToolItemPtr FNavigationToolView::GetRootItem() const
+{
+	if (WeakTool.IsValid())
+	{
+		return WeakTool.Pin()->GetTreeRoot();
+	}
+	return nullptr;
+}
+
+const TArray<FNavigationToolItemPtr>& FNavigationToolView::GetRootVisibleItems() const
+{
+	return RootVisibleItems;
+}
+
+void FNavigationToolView::SaveViewItemFlags(const FNavigationToolItemPtr& InItem, ENavigationToolItemFlags InFlags)
+{
+	if (!InItem.IsValid())
+	{
+		return;
+	}
+
+	const TSharedPtr<FNavigationToolProvider> Provider = InItem->GetProvider();
+	if (!Provider.IsValid())
+	{
+		return;
+	}
+
+	const TSharedPtr<FNavigationTool> Tool = WeakTool.Pin();
+	if (!Tool.IsValid())
+	{
+		return;
+	}
+
+	FNavigationToolSaveState* const SaveState = Provider->GetSaveState(*Tool);
+	if (!SaveState)
+	{
+		UE_LOG(LogSequenceNavigator, Warning, TEXT("SaveViewItemFlags(): Navigation Tool Save State is NULL!"));
+		return;
+	}
+
+	if (!SaveState->ToolViewSaveStates.IsValidIndex(ToolViewId))
+	{
+		UE_LOG(LogSequenceNavigator, Warning, TEXT("SaveViewItemFlags(): '%d' is in invalid saved tool view state index"), ToolViewId);
+		return;
+	}
+
+	if (InFlags == ENavigationToolItemFlags::None)
+	{
+		SaveState->ToolViewSaveStates[ToolViewId].ViewItemFlags.Remove(InItem->GetItemId().GetStringId());
+	}
+	else
+	{
+		SaveState->ToolViewSaveStates[ToolViewId].ViewItemFlags.Add(InItem->GetItemId().GetStringId(), InFlags);
+	}
+}
+
+ENavigationToolItemFlags FNavigationToolView::GetViewItemFlags(const FNavigationToolItemPtr& InItem) const
+{
+	if (!InItem.IsValid())
+	{
+		return ENavigationToolItemFlags::None;
+	}
+
+	const TSharedPtr<FNavigationToolProvider> Provider = InItem->GetProvider();
+	if (!Provider.IsValid())
+	{
+		return ENavigationToolItemFlags::None;
+	}
+
+	const TSharedPtr<FNavigationTool> Tool = WeakTool.Pin();
+	if (!Tool.IsValid())
+	{
+		return ENavigationToolItemFlags::None;
+	}
+
+	FNavigationToolViewSaveState* const ViewSaveState = Provider->GetViewSaveState(*Tool, ToolViewId);
+	if (!ViewSaveState)
+	{
+		return ENavigationToolItemFlags::None;
+	}
+
+	if (const ENavigationToolItemFlags* const OverrideFlags = ViewSaveState->ViewItemFlags.Find(InItem->GetItemId().GetStringId()))
+	{
+		return *OverrideFlags;
+	}
+
+	return ENavigationToolItemFlags::None;
+}
+
+void FNavigationToolView::GetChildrenOfItem(const FNavigationToolItemPtr InItem, TArray<FNavigationToolItemPtr>& OutChildren) const
+{
+	static const TSet<FNavigationToolItemPtr> EmptySet;
+	GetChildrenOfItem(InItem, OutChildren, ENavigationToolItemViewMode::ItemTree, EmptySet);
+}
+
+void FNavigationToolView::GetChildrenOfItem(const FNavigationToolItemPtr& InItem
+	, TArray<FNavigationToolItemPtr>& OutChildren
+	, ENavigationToolItemViewMode InViewMode
+	, const TSet<FNavigationToolItemPtr>& InRecursionDisallowedItems) const
+{
+	if (!InItem.IsValid())
+	{
+		return;
+	}
+
+	const TSharedPtr<FNavigationToolProvider> Provider = InItem->GetProvider();
+	if (!Provider.IsValid() && InItem->GetItemId() != FNavigationToolItemId::RootId)
+	{
+		UE_LOG(LogTemp, Warning, TEXT("Navigation Tool Item Id \"%s\" has no provider, but is a root item!"), *InItem->GetItemId().GetStringId());
+		return;
+	}
+
+	for (const FNavigationToolItemPtr& ChildItem : InItem->GetChildren())
+	{
+		if (!ChildItem.IsValid())
+		{
+			continue;
+		}
+
+		const bool bShouldShowItem = ShouldShowItem(ChildItem, true, InViewMode);
+		if (bShouldShowItem)
+		{
+			// If current item is visible in outliner, add it to the children
+			OutChildren.Add(ChildItem);
+		}
+		else if (!InRecursionDisallowedItems.Contains(ChildItem))
+		{
+			TArray<FNavigationToolItemPtr> GrandChildren;
+
+			// For Muted Hierarchy to be in effect, not only does it have to be on
+			// but also the item should be shown (without counting the filter pass)
+			const bool bShouldUseMutedHierarchy = ToolSettings->ShouldUseMutedHierarchy();
+			const bool bShouldShowItemWithoutFilters = ShouldShowItem(ChildItem, false, InViewMode);
+			const bool bShouldMuteItem = bShouldUseMutedHierarchy && bShouldShowItemWithoutFilters;
+
+			// If it's muted hierarchy, there might be ONLY grand children that are just visible in other view modes, 
+			// so instead of just filtering out the child item, check that there are no grand children from other view modes passing filter tests
+			// If it's NOT muted hierarchy, just get the grand children visible in the requested view mode, as this ChildItem is guaranteed to be hidden
+			const ENavigationToolItemViewMode ViewModeToUse = bShouldMuteItem ? ENavigationToolItemViewMode::All : InViewMode;
+
+			GetChildrenOfItem(ChildItem, GrandChildren, ViewModeToUse, InRecursionDisallowedItems);
+
+			if (!GrandChildren.IsEmpty())
+			{
+				if (bShouldMuteItem)
+				{
+					ReadOnlyItems.Add(ChildItem);	
+					OutChildren.Add(ChildItem);
+				}
+				else
+				{
+					// We can append them knowing that the ViewMode to use is the one passed in and there's no
+					// child that leaked from another view mode
+					ensure(ViewModeToUse == InViewMode);
+					OutChildren.Append(GrandChildren);
+				}
+			}
+		}
+	}
+}
+
+FLinearColor FNavigationToolView::GetItemBrushColor(const FNavigationToolItemPtr InItem) const
+{
+	if (InItem.IsValid())
+	{
+		FLinearColor OutColor = InItem->GetItemBrushColor();
+
+		// If NextSelectedItemIntoView is valid, it means we're scrolling items into view with Next/Previous, 
+		// so Make everything that's not the Current Item a bit more translucent to make the Current Item stand out
+		if (SortedSelectedItems.IsValidIndex(NextSelectedItemIntoView)
+			&& SortedSelectedItems[NextSelectedItemIntoView] != InItem)
+		{
+			OutColor.A *= 0.5f;
+		}
+
+		return OutColor;
+	}
+
+	return FStyleColors::White.GetSpecifiedColor();
+}
+
+TArray<FNavigationToolItemPtr> FNavigationToolView::GetSelectedItems() const
+{
+	return SelectedItems;
+}
+
+int32 FNavigationToolView::GetViewSelectedItemCount() const
+{
+	return SelectedItems.Num();
+}
+
+int32 FNavigationToolView::CalculateVisibleItemCount() const
+{
+	TArray<FNavigationToolItemPtr> RemainingItems = RootVisibleItems;
+	
+	int32 VisibleItemCount = RemainingItems.Num();
+
+	while (RemainingItems.Num() > 0)
+	{
+		const FNavigationToolItemPtr Item = RemainingItems.Pop();
+		
+		TArray<FNavigationToolItemPtr> ChildItems;
+		GetChildrenOfItem(Item, ChildItems);
+		VisibleItemCount += ChildItems.Num();
+		RemainingItems.Append(MoveTemp(ChildItems));
+	}
+
+	// Remove the Read Only Items as they are filtered out items that are still shown because of Hierarchy Viz
+	VisibleItemCount -= ReadOnlyItems.Num();
+
+	return VisibleItemCount;
+}
+
+void FNavigationToolView::SelectItems(TArray<FNavigationToolItemPtr> InItems, ENavigationToolItemSelectionFlags InFlags)
+{
+	// Remove Duplicate Items
+	TSet<FNavigationToolItemPtr> SeenItems;
+	SeenItems.Reserve(InItems.Num());
+	for (TArray<FNavigationToolItemPtr>::TIterator Iter(InItems); Iter; ++Iter)
+	{
+		if (SeenItems.Contains(*Iter))
+		{
+			Iter.RemoveCurrent();
+		}
+		else
+		{
+			SeenItems.Add(*Iter);
+		}
+	}
+
+	// Add the Children of the Items given
+	if (EnumHasAnyFlags(InFlags, ENavigationToolItemSelectionFlags::IncludeChildren))
+	{
+		TArray<FNavigationToolItemPtr> ChildItemsRemaining(InItems);
+		while (ChildItemsRemaining.Num() > 0)
+		{
+			if (FNavigationToolItemPtr ChildItem = ChildItemsRemaining.Pop())
+			{
+				TArray<FNavigationToolItemPtr> Children;
+				GetChildrenOfItem(ChildItem, Children);
+
+				InItems.Append(Children);
+				ChildItemsRemaining.Append(Children);
+			}
+		}
+	}
+
+	if (EnumHasAnyFlags(InFlags, ENavigationToolItemSelectionFlags::AppendToCurrentSelection))
+	{
+		// Remove all repeated items to avoid duplicated entries
+		SelectedItems.RemoveAll([&SeenItems](const FNavigationToolItemPtr& Item)
+		{
+			return SeenItems.Contains(Item);
+		});	
+		TArray<FNavigationToolItemPtr> Items = MoveTemp(SelectedItems);
+		Items.Append(MoveTemp(InItems));
+		InItems = MoveTemp(Items);
+	}
+
+	if (!InItems.IsEmpty() && EnumHasAnyFlags(InFlags, ENavigationToolItemSelectionFlags::ScrollIntoView))
+	{
+		ScrollItemIntoView(InItems[0]);
+	}
+
+	const bool bSignalSelectionChange = EnumHasAnyFlags(InFlags, ENavigationToolItemSelectionFlags::SignalSelectionChange);
+	SetItemSelectionImpl(MoveTemp(InItems), bSignalSelectionChange);
+}
+
+void FNavigationToolView::ClearItemSelection(bool bSignalSelectionChange)
+{
+	SetItemSelectionImpl({}, bSignalSelectionChange);
+}
+
+void FNavigationToolView::SetItemSelectionImpl(TArray<FNavigationToolItemPtr>&& InItems, bool bSignalSelectionChange)
+{
+	/*if (ShouldAutoExpandToSelection())
+	{
+		for (const FNavigationToolItemPtr& Item : InItems)
+		{
+			SetParentItemExpansions(Item, true);
+		}
+	}*/
+
+	SelectedItems = MoveTemp(InItems);
+
+	if (ToolViewWidget.IsValid())
+	{
+		ToolViewWidget->SetItemSelection(SelectedItems, bSignalSelectionChange);
+	}
+	else if (bSignalSelectionChange)
+	{
+		NotifyItemSelectionChanged(SelectedItems, nullptr, true);
+	}
+	
+	Refresh();
+}
+
+void FNavigationToolView::NotifyItemSelectionChanged(const TArray<FNavigationToolItemPtr>& InSelectedItems
+	, const FNavigationToolItemPtr& InItem
+	, bool bUpdateModeTools)
+{
+	if (bSyncingItemSelection)
+	{
+		return;
+	}
+	TGuardValue<bool> Guard(bSyncingItemSelection, true);
+
+	SelectedItems            = InSelectedItems;
+	SortedSelectedItems      = InSelectedItems;
+	NextSelectedItemIntoView = -1;
+
+	FNavigationTool::SortItems(SortedSelectedItems);
+
+	//If we have Pending Items Remaining but we switched Selection via Navigation, treat it as "I want to rename this too"
+	if (bRenamingItems && InItem.IsValid() && InItem != CurrentItemRenaming)
+	{
+		bRequestedRename = true;
+	}
+
+	const TSharedPtr<FNavigationTool> Tool = WeakTool.Pin();
+	if (bUpdateModeTools && Tool.IsValid())
+	{
+		if (ToolSettings->ShouldSyncSelectionToSequencer())
+		{
+			Tool->SyncSequencerSelection(SelectedItems);
+		}
+
+		Tool->SelectItems(SelectedItems, ENavigationToolItemSelectionFlags::None);
+	}
+}
+
+bool FNavigationToolView::IsItemReadOnly(const FNavigationToolItemPtr& InItem) const
+{
+	return ReadOnlyItems.Contains(InItem);;
+}
+
+bool FNavigationToolView::CanSelectItem(const FNavigationToolItemPtr& InItem) const
+{
+	const bool bIsSelectable = InItem.IsValid() && InItem->IsSelectable();
+	return bIsSelectable && !IsItemReadOnly(InItem);
+}
+
+bool FNavigationToolView::IsItemSelected(const FNavigationToolItemPtr& InItem) const
+{
+	return SelectedItems.Contains(InItem);
+}
+
+bool FNavigationToolView::IsItemExpanded(const FNavigationToolItemPtr& InItem, const bool bInUseFilter) const
+{
+	// Don't continue if Item should be hidden in view.
+	// the tree view still calls OnItemExpansionChanged even if it doesn't contain the item
+	// so this preemptive check is needed
+	if (!ShouldShowItem(InItem, bInUseFilter, ENavigationToolItemViewMode::ItemTree))
+	{
+		return false;
+	}
+
+	if (ToolViewWidget.IsValid())
+	{
+		return ToolViewWidget->IsItemExpanded(InItem);
+	}
+
+	return false;
+}
+
+void FNavigationToolView::SetItemExpansion(const FNavigationToolItemPtr& InItem, const bool bInExpand, const bool bInUseFilter)
+{
+	// Don't continue if Item should be hidden in view.
+	// the tree view still calls OnItemExpansionChanged even if it doesn't contain the item
+	// so this preemptive check is needed
+	if (!ShouldShowItem(InItem, bInUseFilter, ENavigationToolItemViewMode::ItemTree))
+	{
+		return;
+	}
+
+	if (ToolViewWidget.IsValid())
+	{
+		ToolViewWidget->SetItemExpansion(InItem, bInExpand);
+	}
+	else
+	{
+		OnItemExpansionChanged(InItem, bInExpand);
+	}
+}
+
+void FNavigationToolView::SetItemExpansionRecursive(const FNavigationToolItemPtr InItem, const bool bInExpand)
+{
+	SetItemExpansion(InItem, bInExpand, false);
+	
+	for (const FNavigationToolItemPtr& Child : InItem->GetChildren())
+	{
+		if (Child.IsValid())
+		{
+			SetItemExpansionRecursive(Child, bInExpand);
+		}
+	}
+}
+
+void FNavigationToolView::SetParentItemExpansions(const FNavigationToolItemPtr& InItem, const bool bInExpand)
+{
+	if (!InItem.IsValid())
+	{
+		return;
+	}
+	
+	TArray<FNavigationToolItemPtr> ItemsToExpand;
+	
+	// Don't auto expand at all if there's a parent preventing it
+	FNavigationToolItemPtr ParentItem = InItem->GetParent();
+	while (ParentItem.IsValid())
+	{
+		if (!ParentItem->CanAutoExpand())
+		{
+			return;
+		}
+		ItemsToExpand.Add(ParentItem);
+		ParentItem = ParentItem->GetParent();
+	}
+
+	for (const FNavigationToolItemPtr& Item : ItemsToExpand)
+	{
+		SetItemExpansion(Item, bInExpand);
+	}
+}
+
+void FNavigationToolView::OnItemExpansionChanged(const FNavigationToolItemPtr InItem, const bool bInIsExpanded)
+{
+	const ENavigationToolItemFlags CurrentFlags = GetViewItemFlags(InItem);
+	
+	ENavigationToolItemFlags TargetFlags = CurrentFlags;
+	
+	if (bInIsExpanded)
+	{
+		TargetFlags |= ENavigationToolItemFlags::Expanded;
+	}
+	else
+	{
+		TargetFlags &= ~ENavigationToolItemFlags::Expanded;
+	}
+
+	SaveViewItemFlags(InItem, TargetFlags);
+
+	if (CurrentFlags != TargetFlags)
+	{
+		InItem->OnExpansionChanged().Broadcast(SharedThis(this), bInIsExpanded);
+	}
+}
+
+bool FNavigationToolView::ShouldShowItem(const FNavigationToolItemPtr& InItem, bool bInUseFilters
+	, ENavigationToolItemViewMode InViewMode) const
+{
+	if (!InItem.IsValid())
+	{
+		return false;
+	}
+
+	if (InItem->IsA<FNavigationToolTreeRoot>())
+	{
+		return true;
+	}
+
+	const TSharedPtr<FNavigationTool> Tool = WeakTool.Pin();
+	if (!Tool.IsValid())
+	{
+		return false;
+	}
+
+	if (!InItem->IsAllowedInTool())
+	{
+		//UE_LOG(LogTemp, Warning, TEXT("Navigation Tool Item Id \"%s\" is hidden: Not Allowed In Tool"), *InItem->GetItemId().GetStringId());
+		return false;
+	}
+
+	if (!InItem->IsViewModeSupported(InViewMode, *this))
+	{
+		//UE_LOG(LogTemp, Warning, TEXT("Navigation Tool Item Id \"%s\" is hidden: View Mode Not Supported"), *InItem->GetItemId().GetStringId());
+		return false;
+	}
+
+	// Allow providers to determine whether the item should be hidden
+	bool bProviderShouldHideItem = false;
+	Tool->ForEachProvider([&bProviderShouldHideItem, &InItem]
+		(const TSharedRef<FNavigationToolProvider>& InToolProvider)
+		{
+			if (InToolProvider->ShouldHideItem(InItem))
+			{
+				bProviderShouldHideItem = true;
+				return false;
+			}
+			return true;
+		});
+	if (bProviderShouldHideItem)
+	{
+		//UE_LOG(LogTemp, Warning, TEXT("Navigation Tool Item Id \"%s\" is hidden: Provider Should Hide Item"), *InItem->GetItemId().GetStringId());
+		return false;
+	}
+
+	// Extra pass for Non-Item Proxies that are parented under an Item Proxy
+	// Hiding an Item Proxy Type should affect all the rest of the items below it
+	if (!InItem->IsA<FNavigationToolItemProxy>())
+	{
+		FNavigationToolItemPtr ItemParent = InItem->GetParent();
+		while (ItemParent.IsValid())
+		{
+			if (ItemParent->IsA<FNavigationToolItemProxy>())
+			{
+				// Stop at the first Item Proxy parent found
+				break;
+			}
+			ItemParent = ItemParent->GetParent();
+		}
+	}
+
+	/** All global filters must fail to hide the item */
+	bool bGlobalFilterOut = false;
+
+	for (const TSharedPtr<FNavigationToolBuiltInFilter>& GlobalFilter : Tool->GlobalFilters)
+	{
+		if (!GlobalFilter->IsActive() && GlobalFilter->PassesFilter(InItem))
+		{
+			bGlobalFilterOut = true;
+			break;
+		}
+	}
+
+	if (bGlobalFilterOut)
+	{
+		return false;
+	}
+
+	if (bInUseFilters && FilterBar->GetFilterData().IsFilteredOut(InItem))
+	{
+		//UE_LOG(LogTemp, Warning, TEXT("Navigation Tool Item Id \"%s\" is hidden: Filtered Out"), *InItem->GetItemId().GetStringId());
+		return false;
+	}
+
+	return true;
+}
+
+int32 FNavigationToolView::GetVisibleChildIndex(const FNavigationToolItemPtr& InParentItem, const FNavigationToolItemPtr& InChildItem) const
+{
+	if (InParentItem.IsValid())
+	{
+		TArray<FNavigationToolItemPtr> Children;
+		GetChildrenOfItem(InParentItem, Children);
+		return Children.Find(InChildItem);
+	}
+	return INDEX_NONE;
+}
+
+FNavigationToolItemPtr FNavigationToolView::GetVisibleChildAt(const FNavigationToolItemPtr& InParentItem, int32 InChildIndex) const
+{
+	if (InParentItem.IsValid())
+	{
+		TArray<FNavigationToolItemPtr> Children;
+		GetChildrenOfItem(InParentItem, Children);
+		if (Children.IsValidIndex(InChildIndex))
+		{
+			return Children[InChildIndex];
+		}
+	}
+	return nullptr;
+}
+
+bool FNavigationToolView::IsToolLocked() const
+{
+	return WeakTool.IsValid() && WeakTool.Pin()->IsToolLocked();
+}
+
+void FNavigationToolView::ShowColumn(const TSharedPtr<INavigationToolColumn>& InColumn)
+{
+	const FName ColumnId = InColumn->GetColumnId();
+	ShowColumnById(ColumnId);
+}
+
+void FNavigationToolView::ShowColumnById(const FName& InColumnId)
+{
+	if (ToolViewWidget.IsValid())
+	{
+		ToolViewWidget->ShowHideColumn(InColumnId, true);
+	}
+
+	SaveColumnState(InColumnId);
+}
+
+void FNavigationToolView::HideColumn(const TSharedPtr<INavigationToolColumn>& InColumn)
+{
+	const FName ColumnId = InColumn->GetColumnId();
+
+	if (ToolViewWidget.IsValid())
+	{
+		ToolViewWidget->ShowHideColumn(ColumnId, false);
+	}
+
+	SaveColumnState(ColumnId);
+}
+
+bool FNavigationToolView::IsColumnVisible(const TSharedPtr<INavigationToolColumn>& InColumn) const
+{
+	const TSharedPtr<FNavigationTool> Tool = WeakTool.Pin();
+	if (!Tool.IsValid())
+	{
+		return false;
+	}
+
+	bool bShouldShow = false;
+
+	Tool->ForEachProvider([this, &InColumn, &Tool, &bShouldShow]
+		(const TSharedRef<FNavigationToolProvider>& InProvider)
+		{
+			if (FNavigationToolViewSaveState* const ViewSaveState = InProvider->GetViewSaveState(*Tool, ToolViewId))
+			{
+				const FName ColumnId = InColumn->GetColumnId();
+				if (const FNavigationToolViewColumnSaveState* const FoundColumnState = ViewSaveState->ColumnsState.Find(ColumnId))
+				{
+					bShouldShow |= FoundColumnState->bVisible;
+				}
+			}
+			return true;
+		});
+
+	return bShouldShow;
+}
+
+ENavigationToolItemViewMode FNavigationToolView::GetItemDefaultViewMode() const
+{
+	return ToolSettings->GetItemDefaultViewMode();
+}
+
+ENavigationToolItemViewMode FNavigationToolView::GetItemProxyViewMode() const
+{
+	return ToolSettings->GetItemProxyViewMode();
+}
+
+void FNavigationToolView::ToggleViewModeSupport(ENavigationToolItemViewMode& InOutViewMode, ENavigationToolItemViewMode InFlags)
+{
+	ToolSettings->ToggleViewModeSupport(InOutViewMode, InFlags);
+	Refresh();
+}
+
+void FNavigationToolView::ToggleItemDefaultViewModeSupport(ENavigationToolItemViewMode InFlags)
+{
+	ToolSettings->ToggleItemDefaultViewModeSupport(InFlags);
+	Refresh();
+}
+
+void FNavigationToolView::ToggleItemProxyViewModeSupport(ENavigationToolItemViewMode InFlags)
+{
+	ToolSettings->ToggleItemProxyViewModeSupport(InFlags);
+	Refresh();
+}
+
+ECheckBoxState FNavigationToolView::GetViewModeCheckState(ENavigationToolItemViewMode InViewMode, ENavigationToolItemViewMode InFlags) const
+{
+	const ENavigationToolItemViewMode Result = InViewMode & InFlags;
+
+	if (Result == InFlags)
+	{
+		return ECheckBoxState::Checked;
+	}
+	
+	if (Result != ENavigationToolItemViewMode::None)
+	{
+		return ECheckBoxState::Undetermined;
+	}
+	
+	return ECheckBoxState::Unchecked;
+}
+
+ECheckBoxState FNavigationToolView::GetItemDefaultViewModeCheckState(ENavigationToolItemViewMode InFlags) const
+{
+	return GetViewModeCheckState(ToolSettings->GetItemDefaultViewMode(), InFlags);
+}
+
+ECheckBoxState FNavigationToolView::GetItemProxyViewModeCheckState(ENavigationToolItemViewMode InFlags) const
+{
+	return GetViewModeCheckState(ToolSettings->GetItemProxyViewMode(), InFlags);
+}
+
+void FNavigationToolView::ToggleMutedHierarchy()
+{
+	ToolSettings->SetUseMutedHierarchy(!ToolSettings->ShouldUseMutedHierarchy());
+	Refresh();
+}
+
+bool FNavigationToolView::IsMutedHierarchyActive() const
+{
+	return ToolSettings->ShouldUseMutedHierarchy();
+}
+
+void FNavigationToolView::ToggleAutoExpandToSelection()
+{
+	ToolSettings->SetAutoExpandToSelection(!ToolSettings->ShouldAutoExpandToSelection());
+	Refresh();
+}
+
+bool FNavigationToolView::ShouldAutoExpandToSelection() const
+{
+	return ToolSettings->ShouldAutoExpandToSelection();
+}
+
+void FNavigationToolView::ToggleUseShortNames()
+{
+	ToolSettings->SetUseShortNames(!ToolSettings->ShouldUseShortNames());
+	Refresh();
+}
+
+bool FNavigationToolView::ShouldUseShortNames() const
+{
+	return ToolSettings->ShouldUseShortNames();
+}
+
+void FNavigationToolView::ToggleShowItemFilters()
+{
+	//Note: Not Marking Navigation Tool Instance as Modified because this is not saved.
+	bShowItemFilters = !bShowItemFilters;
+}
+
+void FNavigationToolView::ToggleShowItemColumns()
+{
+	bShowItemColumns = !bShowItemColumns;
+}
+
+void FNavigationToolView::SetItemTypeHidden(const FName InItemTypeName, const bool bInHidden)
+{
+	if (IsItemTypeHidden(InItemTypeName) != bInHidden)
+	{
+		if (bInHidden)
+		{
+			HiddenItemTypes.Add(InItemTypeName);
+		}
+		else
+		{
+			HiddenItemTypes.Remove(InItemTypeName);
+		}
+		RequestRefresh();
+	}	
+}
+
+void FNavigationToolView::ToggleHideItemTypes(const FName InItemTypeName)
+{
+	SetItemTypeHidden(InItemTypeName, !IsItemTypeHidden(InItemTypeName));
+}
+
+ECheckBoxState FNavigationToolView::GetToggleHideItemTypesState(const FName InItemTypeName) const
+{
+	return IsItemTypeHidden(InItemTypeName) ? ECheckBoxState::Unchecked : ECheckBoxState::Checked;
+}
+
+bool FNavigationToolView::IsItemTypeHidden(const FName InItemTypeName) const
+{
+	return HiddenItemTypes.Contains(InItemTypeName);
+}
+
+bool FNavigationToolView::IsItemTypeHidden(const FNavigationToolItemPtr& InItem) const
+{
+	return IsItemTypeHidden(InItem->GetTypeId().ToName());
+}
+
+void FNavigationToolView::OnDragEnter(const FDragDropEvent& InDragDropEvent, const FNavigationToolItemPtr InTargetItem)
+{
+	if (!InTargetItem.IsValid() && WeakTool.IsValid())
+	{
+		TSharedRef<FNavigationToolItem> TreeRoot = WeakTool.Pin()->GetTreeRoot();
+		const bool bCanAcceptDrop = TreeRoot->CanAcceptDrop(InDragDropEvent, EItemDropZone::OntoItem).IsSet();
+		SetDragIntoTreeRoot(bCanAcceptDrop);
+	}
+	else
+	{
+		SetDragIntoTreeRoot(false);
+	}
+}
+
+void FNavigationToolView::OnDragLeave(const FDragDropEvent& InDragDropEvent, const FNavigationToolItemPtr InTargetItem)
+{
+	// If drag left an item, set the drag into tree root to false (this will set it back to false if a valid item receives DragEnter)
+	SetDragIntoTreeRoot(InTargetItem.IsValid());
+}
+
+FReply FNavigationToolView::OnDragDetected(const FGeometry& InGeometry
+	, const FPointerEvent& InMouseEvent
+	, const FNavigationToolItemPtr InTargetItem)
+{
+	if (!IsToolLocked())
+	{
+		// Only Select Target if it hasn't already been selected
+		if (!IsItemSelected(InTargetItem))
+		{
+			const ENavigationToolItemSelectionFlags SelectionFlags = InMouseEvent.IsControlDown()
+				? ENavigationToolItemSelectionFlags::AppendToCurrentSelection
+				: ENavigationToolItemSelectionFlags::None;
+
+			SelectItems({ InTargetItem }, SelectionFlags);
+		}
+
+		// Get all Selected Items that are in a state where they can be selected again (i.e. not Read Only)
+		TArray<FNavigationToolItemPtr> Items = GetSelectedItems();
+		Items.RemoveAll([this](const FNavigationToolItemPtr& InItem)->bool
+		{
+			return !CanSelectItem(InItem);
+		});
+
+		if (Items.Num() > 0)
+		{
+			const ENavigationToolDragDropActionType ActionType = InMouseEvent.IsAltDown()
+				? ENavigationToolDragDropActionType::Copy
+				: ENavigationToolDragDropActionType::Move;
+
+			const TSharedRef<FNavigationToolItemDragDropOp> DragDropOp = FNavigationToolItemDragDropOp::New(Items, SharedThis(this), ActionType);
+			return FReply::Handled().BeginDragDrop(DragDropOp);
+		}
+	}
+	return FReply::Unhandled();
+}
+
+FReply FNavigationToolView::OnDrop(const FDragDropEvent& InDragDropEvent, EItemDropZone InDropZone, const FNavigationToolItemPtr InTargetItem)
+{
+	SetDragIntoTreeRoot(false);
+
+	if (InTargetItem.IsValid())
+	{
+		return InTargetItem->AcceptDrop(InDragDropEvent, InDropZone);
+	}
+
+	FNavigationToolItemPtr TreeRoot;
+
+	if (WeakTool.IsValid())
+	{
+		TreeRoot = WeakTool.Pin()->GetTreeRoot();
+	}
+
+	if (TreeRoot.IsValid() && TreeRoot->CanAcceptDrop(InDragDropEvent, EItemDropZone::OntoItem))
+	{
+		return TreeRoot->AcceptDrop(InDragDropEvent, EItemDropZone::OntoItem);
+	}
+	
+	return FReply::Unhandled();
+}
+
+TOptional<EItemDropZone> FNavigationToolView::OnCanDrop(const FDragDropEvent& InDragDropEvent
+	, EItemDropZone InDropZone
+	, const FNavigationToolItemPtr InTargetItem) const
+{	
+	if (!IsToolLocked() && InTargetItem.IsValid() && CanSelectItem(InTargetItem))
+	{
+		return InTargetItem->CanAcceptDrop(InDragDropEvent, InDropZone);
+	}
+	return TOptional<EItemDropZone>();
+}
+
+void FNavigationToolView::SetDragIntoTreeRoot(bool bInIsDraggingIntoTreeRoot)
+{
+	if (ToolViewWidget.IsValid())
+	{
+		ToolViewWidget->SetTreeBorderVisibility(bInIsDraggingIntoTreeRoot);
+	}
+}
+
+void FNavigationToolView::RenameSelected()
+{
+	if (const TSharedPtr<FNavigationTool> Tool = WeakTool.Pin())
+	{
+		TArray<FNavigationToolItemPtr> Items = GetSelectedItems();
+		
+		if (Items.IsEmpty())
+		{
+			return;
+		}
+	
+		// Assume we have an item currently renaming
+		ResetRenaming();
+
+		// Remove Items that are Invalid or can't be renamed
+		Items.RemoveAll([](const FNavigationToolItemPtr& InItem)
+			{
+				if (InItem.IsValid())
+				{
+					if (const IRenameableExtension* const RenameableExtension = InItem->CastTo<IRenameableExtension>())
+					{
+						return !RenameableExtension->CanRename();
+					}
+				}
+				return true;
+			});
+		
+		ItemsRemainingRename = MoveTemp(Items);
+
+		if (ItemsRemainingRename.Num() > 0)
+		{
+			FNavigationTool::SortItems(ItemsRemainingRename);
+			bRenamingItems = true;
+		}
+	}
+}
+
+void FNavigationToolView::ResetRenaming()
+{
+	if (CurrentItemRenaming.IsValid())
+	{
+		CurrentItemRenaming->OnRenameAction().RemoveAll(this);
+		CurrentItemRenaming.Reset();
+	}
+
+	if (ItemsRemainingRename.IsEmpty())
+	{
+		bRenamingItems = false;
+	}
+}
+
+void FNavigationToolView::OnItemRenameAction(const ENavigationToolRenameAction InRenameAction, const TSharedPtr<INavigationToolView>& InToolView)
+{
+	if (InToolView.Get() != this)
+	{
+		return;
+	}
+	
+	switch (InRenameAction)
+	{
+	case ENavigationToolRenameAction::None:
+		break;
+	
+	case ENavigationToolRenameAction::Requested:
+		break;
+	
+	case ENavigationToolRenameAction::Cancelled:
+		ItemsRemainingRename.Reset();
+		ResetRenaming();
+		break;
+	
+	case ENavigationToolRenameAction::Completed:
+		ResetRenaming();
+		break;
+	
+	default:
+		break;
+	}
+}
+
+bool FNavigationToolView::CanRenameSelected() const
+{
+	for (const FNavigationToolItemPtr& Item : SelectedItems)
+	{
+		if (!Item.IsValid())
+		{
+			return false;
+		}
+
+		if (const IRenameableExtension* const RenameableExtension = Item->CastTo<IRenameableExtension>())
+		{
+			if (!RenameableExtension->CanRename())
+			{
+				return false;
+			}
+		}
+	}
+
+	return SelectedItems.Num() > 0;
+}
+
+void FNavigationToolView::DeleteSelected()
+{
+	if (const TSharedPtr<FNavigationTool> Tool = WeakTool.Pin())
+	{
+		TArray<FNavigationToolItemPtr> Items = GetSelectedItems();
+
+		Items.RemoveAll([](const FNavigationToolItemPtr& InItem)
+		{
+			return !InItem.IsValid() || !InItem->CanDelete();
+		});
+
+		if (Items.IsEmpty())
+		{
+			return;
+		}
+
+		Tool->DeleteItems(Items);
+	}
+}
+
+bool FNavigationToolView::CanDeleteSelected() const
+{
+	for (const FNavigationToolItemPtr& Item : GetSelectedItems())
+	{
+		if (Item && Item->CanDelete())
+		{
+			return true;
+		}
+	}
+
+	return false;
+}
+
+void FNavigationToolView::DuplicateSelected()
+{
+	if (const TSharedPtr<FNavigationTool> Tool = WeakTool.Pin())
+	{
+		//Tool->DuplicateItems(GetSelectedItems(), nullptr, TOptional<EItemDropZone>());
+	}
+}
+
+bool FNavigationToolView::CanDuplicateSelected() const
+{
+	for (const FNavigationToolItemPtr& Item : GetSelectedItems())
+	{
+		if (Item && Item->IsA<FNavigationToolActor>())
+		{
+			return true;
+		}
+	}
+	return false;
+}
+
+void FNavigationToolView::SelectChildren(bool bIsRecursive)
+{
+	TArray<FNavigationToolItemPtr> ItemsToSelect;
+	TArray<FNavigationToolItemPtr> RemainingItems = GetSelectedItems();
+
+	while (RemainingItems.Num() > 0)
+	{
+		//Note: Pop here will affect order of Children in Selection
+		const FNavigationToolItemPtr ParentItem = RemainingItems.Pop();
+
+		TArray<FNavigationToolItemPtr> ChildItems;
+		GetChildrenOfItem(ParentItem, ChildItems);
+		if (bIsRecursive)
+		{
+			RemainingItems.Append(ChildItems);
+		}
+		ItemsToSelect.Append(ChildItems);
+	}
+
+	SelectItems(ItemsToSelect
+		, ENavigationToolItemSelectionFlags::AppendToCurrentSelection | ENavigationToolItemSelectionFlags::SignalSelectionChange);
+}
+
+bool FNavigationToolView::CanSelectChildren() const
+{
+	return GetViewSelectedItemCount() > 0;
+}
+
+void FNavigationToolView::SelectParent()
+{
+	const TSet<FNavigationToolItemPtr> Items(GetSelectedItems());
+
+	TSet<FNavigationToolItemPtr> ParentItemsToSelect;
+	ParentItemsToSelect.Reserve(Items.Num());
+
+	const FNavigationToolItemPtr RootItem = GetRootItem();
+	
+	//Add only Valid Parents that are not Root and are not part of the Original Selection!
+	for (const FNavigationToolItemPtr& Item : Items)
+	{
+		if (Item.IsValid())
+		{
+			const FNavigationToolItemPtr ParentItem = Item->GetParent();
+			if (ParentItem.IsValid() && ParentItem != RootItem && !Items.Contains(ParentItem))
+			{
+				ParentItemsToSelect.Add(ParentItem);
+			}
+		}
+	}
+
+	SortAndSelectItems(ParentItemsToSelect.Array());
+}
+
+bool FNavigationToolView::CanSelectParent() const
+{
+	return GetViewSelectedItemCount() == 1;
+}
+
+void FNavigationToolView::SelectFirstChild()
+{
+	const TArray<FNavigationToolItemPtr> Items = GetSelectedItems();
+
+	TSet<FNavigationToolItemPtr> FirstChildItemsToSelect;
+	FirstChildItemsToSelect.Reserve(Items.Num());
+
+	for (const FNavigationToolItemPtr& Item : Items)
+	{
+		if (Item.IsValid())
+		{
+			const FNavigationToolItemPtr FirstChildItem = GetVisibleChildAt(Item, 0);
+
+			//Don't select Component items! (Component items on selection also select their owner actor items, which can cause undesired issues)
+			if (FirstChildItem.IsValid() && !FirstChildItem->IsA<FNavigationToolComponent>())
+			{
+				FirstChildItemsToSelect.Add(FirstChildItem);
+			}
+		}
+	}
+	
+	SortAndSelectItems(FirstChildItemsToSelect.Array());
+}
+
+bool FNavigationToolView::CanSelectFirstChild() const
+{
+	return GetViewSelectedItemCount() == 1;
+}
+
+void FNavigationToolView::SelectSibling(int32 InDeltaIndex)
+{
+	const TArray<FNavigationToolItemPtr> Items = GetSelectedItems();
+
+	TSet<FNavigationToolItemPtr> SiblingItemsToSelect;
+	SiblingItemsToSelect.Reserve(Items.Num());
+
+	for (const FNavigationToolItemPtr& Item : Items)
+	{
+		if (Item.IsValid() && Item->GetParent())
+		{
+			const FNavigationToolItemPtr ParentItem = Item->GetParent();
+
+			const int32 ItemIndex   = GetVisibleChildIndex(ParentItem, Item);
+			const int32 TargetIndex = ItemIndex + InDeltaIndex;
+
+			//Don't try to Normalize Index, if it's Invalid, we won't cycle and just skip that selection
+			const FNavigationToolItemPtr SiblingToSelect = GetVisibleChildAt(ParentItem, TargetIndex);
+
+			//Don't select Component items! (Component items on selection also select their owner actor items, which can cause undesired issues)
+			if (SiblingToSelect.IsValid() && !SiblingToSelect->IsA<FNavigationToolComponent>())
+			{
+				SiblingItemsToSelect.Add(SiblingToSelect);
+			}
+		}
+	}
+	SortAndSelectItems(SiblingItemsToSelect.Array());
+}
+
+bool FNavigationToolView::CanSelectSibling() const
+{
+	return GetViewSelectedItemCount() == 1;
+}
+
+void FNavigationToolView::ExpandAll()
+{
+	for (const FNavigationToolItemPtr& Item : RootVisibleItems)
+	{
+		SetItemExpansionRecursive(Item, true);
+	}
+}
+
+bool FNavigationToolView::CanExpandAll() const
+{
+	return true;
+}
+
+void FNavigationToolView::CollapseAll()
+{
+	for (const FNavigationToolItemPtr& Item : RootVisibleItems)
+	{
+		SetItemExpansionRecursive(Item, false);
+	}
+}
+
+bool FNavigationToolView::CanCollapseAll() const
+{
+	return true;
+}
+
+void FNavigationToolView::ScrollNextIntoView()
+{
+	ScrollDeltaIndexIntoView(+1);
+}
+
+void FNavigationToolView::ScrollPrevIntoView()
+{
+	ScrollDeltaIndexIntoView(-1);
+}
+
+bool FNavigationToolView::CanScrollNextIntoView() const
+{
+	return GetViewSelectedItemCount() > 0;
+}
+
+void FNavigationToolView::ScrollDeltaIndexIntoView(int32 DeltaIndex)
+{
+	const int32 SelectedItemCount = SortedSelectedItems.Num();
+	if (SortedSelectedItems.Num() > 0)
+	{
+		const int32 TargetIndex  = NextSelectedItemIntoView + DeltaIndex;
+		NextSelectedItemIntoView = TargetIndex % SelectedItemCount;
+		if (NextSelectedItemIntoView < 0)
+		{
+			NextSelectedItemIntoView += SelectedItemCount;
+		}
+		ScrollItemIntoView(SortedSelectedItems[NextSelectedItemIntoView]);
+	}
+}
+
+void FNavigationToolView::ScrollItemIntoView(const FNavigationToolItemPtr& InItem)
+{
+	if (InItem.IsValid())
+	{
+		SetParentItemExpansions(InItem, true);
+		if (ToolViewWidget.IsValid() && ToolViewWidget->GetTreeView())
+		{
+			ToolViewWidget->GetTreeView()->FocusOnItem(InItem);
+			ToolViewWidget->ScrollItemIntoView(InItem);
+		}
+	}
+}
+
+void FNavigationToolView::SortAndSelectItems(TArray<FNavigationToolItemPtr> InItemsToSelect)
+{
+	if (!InItemsToSelect.IsEmpty())
+	{
+		FNavigationTool::SortItems(InItemsToSelect);
+
+		SelectItems(InItemsToSelect
+			, ENavigationToolItemSelectionFlags::SignalSelectionChange
+			| ENavigationToolItemSelectionFlags::ScrollIntoView);
+	}
+}
+
+void FNavigationToolView::RefreshTool(bool bInImmediateRefresh)
+{
+	if (const TSharedPtr<INavigationTool> Tool = GetOwnerTool())
+	{
+		if (bInImmediateRefresh)
+		{
+			Tool->Refresh();
+		}
+		else
+		{
+			Tool->RequestRefresh();
+		}
+	}
+}
+
+void FNavigationToolView::EnsureToolViewCount(const int32 InToolViewId) const
+{
+	const TSharedPtr<FNavigationTool> Tool = WeakTool.Pin();
+	if (!Tool.IsValid())
+	{
+		return;
+	}
+
+	FNavigationTool& ToolRef = *Tool;
+
+	ToolRef.ForEachProvider([&ToolRef, InToolViewId]
+		(const TSharedRef<FNavigationToolProvider>& InToolProvider)
+		{
+			InToolProvider->EnsureToolViewCount(ToolRef, InToolViewId);
+			return true;
+		});
+}
+
+void FNavigationToolView::SaveViewState(const TSharedRef<FNavigationToolProvider>& InProvider)
+{
+	const TSharedPtr<FNavigationTool> Tool = WeakTool.Pin();
+	if (!Tool.IsValid())
+	{
+		return;
+	}
+
+	const FNavigationTool& ToolRef = *Tool;
+
+	EnsureToolViewCount(ToolViewId);
+
+	if (FNavigationToolViewSaveState* const ViewSaveState = InProvider->GetViewSaveState(ToolRef, ToolViewId))
+	{
+		// Save view state filters
+		ViewSaveState->ActiveItemFilters.Reset();
+		for (const TSharedRef<FNavigationToolFilter>& ActiveItemFilter : FilterBar->GetActiveFilters())
+		{
+			ViewSaveState->ActiveItemFilters.Add(*ActiveItemFilter->GetDisplayName().ToString());
+		}
+
+		SaveColumnState();
+		SaveToolViewItems(*ViewSaveState);
+	}
+}
+
+void FNavigationToolView::LoadViewState(const TSharedRef<FNavigationToolProvider>& InProvider)
+{
+	const TSharedPtr<FNavigationTool> Tool = WeakTool.Pin();
+	if (!Tool.IsValid())
+	{
+		return;
+	}
+
+	const FNavigationTool& ToolRef = *Tool;
+
+	EnsureToolViewCount(ToolViewId);
+
+	// Disable all filters before load
+	FilterBar->EnableAllFilters(false, {});
+
+	if (FNavigationToolViewSaveState* const ViewSaveState = InProvider->GetViewSaveState(ToolRef, ToolViewId))
+	{
+		LoadFilterState(*ViewSaveState, false, false);
+		LoadToolViewItems(*ViewSaveState);
+	}
+
+	PostLoad();
+
+	FilterBar->RequestFilterUpdate();
+}
+
+void FNavigationToolView::SaveColumnState(const FName InColumnId)
+{
+	if (!ToolViewWidget.IsValid())
+	{
+		return;
+	}
+
+	const TSharedPtr<FNavigationTool> Tool = WeakTool.Pin();
+	if (!Tool.IsValid())
+	{
+		return;
+	}
+
+	const FNavigationTool& ToolRef = *Tool;
+
+	// Save all columns unless a specific column is specified
+	if (InColumnId.IsNone())
+	{
+		// Save each column to their respective providers save data
+		// Note: Some columns may have multiple providers
+		for (const TTuple<FName, TSharedPtr<INavigationToolColumn>>& Column : Columns)
+		{
+			if (!Column.Value.IsValid())
+			{
+				continue;
+			}
+
+			const FName ColumnId = Column.Value->GetColumnId();
+
+			Tool->ForEachProvider([this, &ColumnId, &ToolRef]
+				(const TSharedRef<FNavigationToolProvider>& InProvider)
+				{
+					if (FNavigationToolViewSaveState* const ViewSaveState = InProvider->GetViewSaveState(ToolRef, ToolViewId))
+					{
+						FNavigationToolViewColumnSaveState& FoundColumnState = ViewSaveState->ColumnsState.FindOrAdd(ColumnId);
+						ToolViewWidget->GenerateColumnState(ColumnId, FoundColumnState);
+					}
+					return true;
+				});
+		}
+	}
+	else if (ensure(Columns.Contains(InColumnId)))
+	{
+		Tool->ForEachProvider([this, &InColumnId, &ToolRef]
+			(const TSharedRef<FNavigationToolProvider>& InProvider)
+			{
+				// Save the specific column to its providers save data
+				const TSharedPtr<INavigationToolColumn>& Column = Columns[InColumnId];
+				if (FNavigationToolViewSaveState* const ViewSaveState = InProvider->GetViewSaveState(ToolRef, ToolViewId))
+				{
+					FNavigationToolViewColumnSaveState& FoundColumnState = ViewSaveState->ColumnsState.FindOrAdd(InColumnId);
+					ToolViewWidget->GenerateColumnState(InColumnId, FoundColumnState);
+				}
+				return true;
+			});
+	}
+}
+
+void FNavigationToolView::SaveFilterState(FNavigationToolViewSaveState& OutViewSaveState)
+{
+	OutViewSaveState.ActiveItemFilters.Reset();
+
+	for (const TSharedRef<FNavigationToolFilter>& ActiveItemFilter : FilterBar->GetActiveFilters())
+	{
+		OutViewSaveState.ActiveItemFilters.Add(*ActiveItemFilter->GetDisplayName().ToString());
+	}
+}
+
+void FNavigationToolView::LoadFilterState(const FNavigationToolViewSaveState& InViewSaveState
+	, const bool bInDisableAllFilters
+	, const bool bInRequestFilterUpdate)
+{
+	if (bInDisableAllFilters)
+	{
+		FilterBar->EnableAllFilters(false, {});
+	}
+
+	for (const FName ActiveItemFilterName : InViewSaveState.ActiveItemFilters)
+	{
+		FilterBar->SetFilterActiveByDisplayName(ActiveItemFilterName.ToString(), true, false);
+	}
+
+	if (bInRequestFilterUpdate)
+	{
+		FilterBar->RequestFilterUpdate();
+	}
+}
+
+void FNavigationToolView::SaveToolViewItems(FNavigationToolViewSaveState& OutViewSaveState)
+{
+	const TSharedRef<FNavigationToolItem> TreeRoot = GetOwnerTool()->GetTreeRoot();
+
+	TArray<FNavigationToolItemPtr> ItemsToSave = TreeRoot->GetChildren();
+
+	OutViewSaveState.ViewItemFlags.Reset();
+
+	while (ItemsToSave.Num() > 0)
+	{
+		FNavigationToolItemPtr ItemToSave = ItemsToSave.Pop();
+		if (ItemToSave.IsValid())
+		{
+			// Iteratively also Save children
+			ItemsToSave.Append(ItemToSave->GetChildren());
+
+			const FNavigationToolItemId ItemId = ItemToSave->GetItemId();
+			const FString StringId = ItemId.GetStringId();
+
+			// Save Item State Flags
+			if (const ENavigationToolItemFlags* const ItemFlags = OutViewSaveState.ViewItemFlags.Find(StringId))
+			{
+				OutViewSaveState.ViewItemFlags.Add(StringId, *ItemFlags);
+			}
+			else
+			{
+				OutViewSaveState.ViewItemFlags.Remove(StringId);
+			}
+		}
+	}
+}
+
+void FNavigationToolView::LoadToolViewItems(FNavigationToolViewSaveState& InViewSaveState)
+{
+	const TSharedRef<FNavigationToolItem> TreeRoot = GetOwnerTool()->GetTreeRoot();
+
+	TArray<FNavigationToolItemPtr> ItemsToLoad = TreeRoot->GetChildren();
+
+	while (ItemsToLoad.Num() > 0)
+	{
+		if (const FNavigationToolItemPtr ItemToLoad = ItemsToLoad.Pop())
+		{
+			// Iteratively also Load Children
+			ItemsToLoad.Append(ItemToLoad->GetChildren());
+
+			const FNavigationToolItemId ItemId = ItemToLoad->GetItemId();
+			const FString StringId = ItemId.GetStringId();
+
+			// Load Item Flags
+			if (const ENavigationToolItemFlags* const ItemFlags = InViewSaveState.ViewItemFlags.Find(StringId))
+			{
+				InViewSaveState.ViewItemFlags.Add(StringId, *ItemFlags);
+			}
+			else
+			{
+				InViewSaveState.ViewItemFlags.Remove(StringId);
+			}
+		}
+	}
+}
+
+TSharedRef<SWidget> FNavigationToolView::GetColumnMenuContent(const FName InColumnId)
+{
+	FMenuBuilder MenuBuilder(true, ViewCommandList);
+
+	const FNavigationToolCommands& ToolCommands = FNavigationToolCommands::Get();
+
+	MenuBuilder.AddMenuEntry(
+		LOCTEXT("ResetColumnSize", "Reset Column Size"), 
+		LOCTEXT("ResetColumnSizeTooltip", "Resets the size of this column to the default"), 
+		FSlateIcon(),
+		FUIAction(FExecuteAction::CreateRaw(this, &FNavigationToolView::ResetColumnSize, InColumnId),
+			FCanExecuteAction::CreateRaw(this, &FNavigationToolView::CanResetColumnSize, InColumnId)));
+
+	MenuBuilder.AddMenuSeparator();
+
+	MenuBuilder.AddMenuEntry(ToolCommands.ResetVisibleColumnSizes);
+
+	return MenuBuilder.MakeWidget();
+}
+
+void FNavigationToolView::ResetColumnSize(const FName InColumnId)
+{
+	if (!ToolViewWidget.IsValid())
+	{
+		return;
+	}
+
+	if (!Columns.Contains(InColumnId) || !ToolViewWidget->IsColumnVisible(InColumnId))
+	{
+		return;
+	}
+
+	const TSharedPtr<FNavigationTool> Tool = WeakTool.Pin();
+	if (!Tool.IsValid())
+	{
+		return;
+	}
+
+	const FNavigationTool& ToolRef = *Tool;
+
+	Tool->ForEachProvider([this, &InColumnId, &ToolRef]
+		(const TSharedRef<FNavigationToolProvider>& InProvider)
+		{
+			if (FNavigationToolViewSaveState* const ViewSaveState = InProvider->GetViewSaveState(ToolRef, ToolViewId))
+			{
+				const float DefaultSize = Columns[InColumnId]->GetFillWidth();
+				if (DefaultSize > 0.f)
+				{
+					ViewSaveState->ColumnsState.FindOrAdd(InColumnId).Size = DefaultSize;
+
+					ToolViewWidget->SetColumnWidth(InColumnId, DefaultSize);
+
+					ToolViewWidget->GenerateColumnState(InColumnId, ViewSaveState->ColumnsState.FindOrAdd(InColumnId));
+				}
+			}
+			return true;
+		});
+}
+
+bool FNavigationToolView::CanResetColumnSize(const FName InColumnId) const
+{
+	if (!Columns.Contains(InColumnId) || !ToolViewWidget->IsColumnVisible(InColumnId))
+	{
+		return false;
+	}
+
+	const TSharedPtr<FNavigationTool> Tool = WeakTool.Pin();
+	if (!Tool.IsValid())
+	{
+		return false;
+	}
+
+	const FNavigationTool& ToolRef = *Tool;
+
+	bool bCanReset = false;
+
+	Tool->ForEachProvider([this, &InColumnId, &ToolRef, &bCanReset]
+		(const TSharedRef<FNavigationToolProvider>& InProvider)
+		{
+			const FNavigationToolViewSaveState* const ViewSaveState = InProvider->GetViewSaveState(ToolRef, ToolViewId);
+			if (!ViewSaveState || !ViewSaveState->ColumnsState.Contains(InColumnId))
+			{
+				return true;
+			}
+
+			const float DefaultSize = Columns[InColumnId]->GetFillWidth();
+			if (DefaultSize <= 0.f)
+			{
+				return true;
+			}
+
+			bCanReset |= ViewSaveState->ColumnsState[InColumnId].Size != DefaultSize;
+
+			return true;
+		});
+
+	return bCanReset;
+}
+
+void FNavigationToolView::ResetVisibleColumnSizes()
+{
+	if (!ToolViewWidget.IsValid())
+	{
+		return;
+	}
+
+	const TSharedPtr<FNavigationTool> Tool = WeakTool.Pin();
+	if (!Tool.IsValid())
+	{
+		return;
+	}
+
+	const FNavigationTool& ToolRef = *Tool;
+
+	for (const TPair<FName, TSharedPtr<INavigationToolColumn>>& Pair : Columns)
+	{
+		const FName& ColumnId = Pair.Key;
+
+		if (!ToolViewWidget->IsColumnVisible(ColumnId))
+		{
+			continue;
+		}
+
+		const float DefaultSize = Columns[ColumnId]->GetFillWidth();
+		if (DefaultSize <= 0.f)
+		{
+			continue;
+		}
+
+		Tool->ForEachProvider([this, &ColumnId, &ToolRef, DefaultSize]
+			(const TSharedRef<FNavigationToolProvider>& InProvider)
+			{
+				if (FNavigationToolViewSaveState* const ViewSaveState = InProvider->GetViewSaveState(ToolRef, ToolViewId))
+				{
+					ToolViewWidget->SetColumnWidth(ColumnId, DefaultSize);
+					ToolViewWidget->GenerateColumnState(ColumnId, ViewSaveState->ColumnsState.FindOrAdd(ColumnId));
+				}
+				return true;
+			});
+	}
+}
+
+bool FNavigationToolView::CanResetAllColumnSizes() const
+{
+	return true;
+}
+
+void FNavigationToolView::SaveNewCustomColumnView()
+{
+	TSet<FNavigationToolColumnView>& CustomColumnViews = ToolSettings->GetCustomColumnViews();
+
+	SaveColumnState();
+
+	// Create a unique column view name suggestion
+	auto DoesColumnViewExist = [&CustomColumnViews](const FText& InViewName)
+	{
+		for (const FNavigationToolColumnView& ColumnView : CustomColumnViews)
+		{
+			if (ColumnView.ViewName.EqualTo(InViewName))
+			{
+				return true;
+			}
+		}
+		return false;
+	};
+
+	FNavigationToolColumnView NewColumnView;
+
+	for (int32 Index = 1; Index < INT_MAX; ++Index)
+	{
+		NewColumnView.ViewName = FText::Format(LOCTEXT("ColumnViewName", "Column View {0}"), { Index });
+		if (!DoesColumnViewExist(NewColumnView.ViewName))
+		{
+			break;
+		}
+	}
+
+	for (const TPair<FName, TSharedPtr<INavigationToolColumn>>& Column : Columns)
+	{
+		const FName ColumnId = Column.Value->GetColumnId();
+		if (IsColumnVisible(Column.Value))
+		{
+			NewColumnView.VisibleColumns.Add(ColumnId);
+		}
+	}
+
+	// Prompt user for name, using the generated unique suggestion as the default name
+	FModalTextInputDialog InputDialog;
+	InputDialog.InputLabel = LOCTEXT("NewColumnViewName_InputLabel", "New Column View Name");
+	if (!InputDialog.Open(NewColumnView.ViewName, NewColumnView.ViewName))
+	{
+		return;
+	}
+
+	bool bAlreadyExists = false;
+	CustomColumnViews.Add(NewColumnView, &bAlreadyExists);
+
+	if (bAlreadyExists)
+	{
+		FMessageDialog::Open(EAppMsgCategory::Error, EAppMsgType::Ok
+			, LOCTEXT("AlreadyExistsErrorText", "Column view name already exists!"));
+		return;
+	}
+
+	CustomColumnViews.Sort([](const FNavigationToolColumnView& InA, const FNavigationToolColumnView& InB)
+		{
+			return InA.ViewName.CompareTo(InB.ViewName) < 0;
+		});
+
+	ToolSettings->SaveConfig();
+}
+
+void FNavigationToolView::ApplyCustomColumnView(const FText InColumnViewName)
+{
+	if (!ToolViewWidget.IsValid() || InColumnViewName.IsEmptyOrWhitespace())
+	{
+		return;
+	}
+
+	const FNavigationToolColumnView* const SavedColumnView = ToolSettings->FindCustomColumnView(InColumnViewName);
+	if (!SavedColumnView)
+	{
+		return;
+	}
+
+	for (const TPair<FName, TSharedPtr<INavigationToolColumn>>& Column : Columns)
+	{
+		const FName ColumnId = Column.Value->GetColumnId();
+		const bool bColumnVisible = SavedColumnView->VisibleColumns.Contains(ColumnId);
+		ToolViewWidget->ShowHideColumn(ColumnId, bColumnVisible);
+	}
+
+	SaveColumnState();
+}
+
+bool FNavigationToolView::CanFocusSingleSelection() const
+{
+	if (SelectedItems.Num() == 1)
+	{
+		const FNavigationToolSequence* const SequenceItem = SelectedItems[0]->CastTo<FNavigationToolSequence>();
+		return SequenceItem && IsValid(SequenceItem->GetSequence());
+	}
+	return false;
+}
+
+void FNavigationToolView::FocusSingleSelection()
+{
+	if (SelectedItems.Num() != 1)
+	{
+		return;
+	}
+
+	const FNavigationToolSequence* const SequenceItem = SelectedItems[0]->CastTo<FNavigationToolSequence>();
+	if (!SequenceItem)
+	{
+		return;
+	}
+
+	UMovieSceneSequence* const Sequence = SequenceItem->GetSequence();
+	if (!IsValid(Sequence))
+	{
+		return;
+	}
+
+	const TSharedPtr<INavigationTool> Tool = GetOwnerTool();
+	if (!Tool.IsValid())
+	{
+		return;
+	}
+
+	FocusSequence(*Tool, *Sequence);
+}
+
+bool FNavigationToolView::CanFocusInContentBrowser() const
+{
+	if (SelectedItems.Num() != 1)
+	{
+		return false;
+	}
+
+	const FNavigationToolSequence* const SequenceItem = SelectedItems[0]->CastTo<FNavigationToolSequence>();
+	if (!SequenceItem)
+	{
+		return false;
+	}
+
+	UMovieSceneSequence* const Sequence = SequenceItem->GetSequence();
+	if (!IsValid(Sequence))
+	{
+		return false;
+	}
+
+	return Cast<ULevelSequence>(Sequence) != nullptr;
+}
+
+void FNavigationToolView::FocusInContentBrowser()
+{
+	if (SelectedItems.Num() != 1)
+	{
+		return;
+	}
+
+	const FNavigationToolSequence* const SequenceItem = SelectedItems[0]->CastTo<FNavigationToolSequence>();
+	if (!SequenceItem)
+	{
+		return;
+	}
+
+	UMovieSceneSequence* const Sequence = SequenceItem->GetSequence();
+	if (!IsValid(Sequence))
+	{
+		return;
+	}
+
+	ULevelSequence* const LevelSequence = Cast<ULevelSequence>(Sequence);
+	if (!IsValid(LevelSequence))
+	{
+		return;
+	}
+
+	const TArray<UObject*> ObjectsToSync = { LevelSequence };
+
+	FContentBrowserModule& ContentBrowserModule = FModuleManager::Get().LoadModuleChecked<FContentBrowserModule>(TEXT("ContentBrowser"));
+	ContentBrowserModule.Get().SyncBrowserToAssets(ObjectsToSync);
+}
+
+bool FNavigationToolView::UpdateFilters()
+{
+	if (!bFilterUpdateRequested)
+	{
+		return false;
+	}
+
+	const TSharedPtr<INavigationTool> Tool = GetOwnerTool();
+	if (!Tool.IsValid())
+	{
+		return false;
+	}
+
+	const FNavigationToolFilterData& PreviousFilterData = FilterBar->GetFilterData();
+	const FNavigationToolFilterData& FilterData = FilterBar->FilterNodes();
+
+	//bFilteringOnNodeGroups = Sequencer.GetFocusedMovieSceneSequence()->GetMovieScene()->GetNodeGroups().HasAnyActiveFilter();
+	bFilterUpdateRequested = false;
+
+	// Return whether the new list of FilteredNodes is different than the previous list
+	return (PreviousFilterData.GetDisplayNodeCount() != FilterData.GetDisplayNodeCount()
+		|| PreviousFilterData != FilterData);
+}
+
+}
+
+#undef LOCTEXT_NAMESPACE
diff --git a/Engine/Plugins/Experimental/SequenceNavigator/Source/SequenceNavigator/Private/NavigationToolView.h b/Engine/Plugins/Experimental/SequenceNavigator/Source/SequenceNavigator/Private/NavigationToolView.h
new file mode 100644
--- /dev/null
+++ b/Engine/Plugins/Experimental/SequenceNavigator/Source/SequenceNavigator/Private/NavigationToolView.h
@@ -0,0 +1,434 @@
+// Copyright Epic Games, Inc. All Rights Reserved.
+
+#pragma once
+
+#include "MovieScene.h"
+#include "Filters/NavigationToolFilterBar.h"
+#include "INavigationToolView.h"
+#include "NavigationToolDefines.h"
+#include "NavigationToolItemType.h"
+
+class FDragDropEvent;
+class FReply;
+class FUICommandList;
+class SWidget;
+class UNavigationToolSettings;
+enum class ECheckBoxState : uint8;
+enum class EItemDropZone;
+struct FGeometry;
+struct FNavigationToolSaveState;
+struct FPointerEvent;
+
+namespace UE::Sequencer::NavigationTool
+{
+
+class FNavigationTool;
+class FNavigationToolFilterBar;
+class FNavigationToolItemContextMenu;
+class FNavigationToolProvider;
+class INavigationToolColumn;
+class SNavigationToolView;
+
+/**
+  * A view instance of the Navigation Tool, that handles viewing a subset of the outliner items based on
+  * item filters, search text, hierarchy type, etc
+  */
+class FNavigationToolView : public INavigationToolView
+{
+private:
+	friend FNavigationToolSaveState;
+	friend FNavigationTool;
+
+	struct FPrivateToken { explicit FPrivateToken() = default; };
+
+public:
+	explicit FNavigationToolView(FPrivateToken);
+	virtual ~FNavigationToolView() override;
+
+	/** Creates columns from a specific provider */
+	void CreateColumns(const TSharedRef<FNavigationToolProvider>& InProvider);
+
+	void CreateDefaultColumnViews(const TSharedRef<FNavigationToolProvider>& InProvider);
+
+	/**
+	 * Creates a Navigation Tool View Instance and register it to the Navigation Tool
+	 * @param InToolViewId the Id assigned to the Instance created
+	 * @param InTool the Navigation Tool this Instance View belongs to
+	 * @param InBaseCommandList the base command list
+	 */
+	static TSharedRef<FNavigationToolView> CreateInstance(const int32 InToolViewId
+		, const TSharedRef<FNavigationTool>& InTool
+		, const TSharedPtr<FUICommandList>& InBaseCommandList);
+
+	void PostLoad();
+
+	/** Called when the UNavigationToolSettings has a property change */
+	void OnToolSettingsChanged(UObject* InObject, FPropertyChangedEvent& InPropertyChangedEvent);
+
+	void Tick(float InDeltaTime);
+
+	void BindCommands(const TSharedPtr<FUICommandList>& InBaseCommandList);
+
+	TSharedPtr<FUICommandList> GetBaseCommandList() const;
+	TSharedPtr<FUICommandList> GetViewCommandList() const { return ViewCommandList; }
+
+	/** Notifies FNavigationTool so that this view instance becomes the most recent interacted outliner view */
+	void UpdateRecentViews();
+
+	/** Returns whether this view instance is the most recently interacted between the instances in FNavigationTool */
+	bool IsMostRecentToolView() const;
+
+	int32 GetToolViewId() const { return ToolViewId; }
+
+	//~ Begin INavigationToolView
+	virtual TSharedPtr<INavigationTool> GetOwnerTool() const override;
+	virtual TSharedPtr<SWidget> GetToolWidget() const override;
+	virtual TSharedPtr<ISequencer> GetSequencer() const override;
+	virtual void RequestRefresh() override;
+	virtual FOnToolViewRefreshed& GetOnToolViewRefreshed() override { return OnToolViewRefreshed; }
+	virtual void SetKeyboardFocus() override;
+	virtual ENavigationToolItemViewMode GetItemDefaultViewMode() const override;
+	virtual ENavigationToolItemViewMode GetItemProxyViewMode() const override;
+	virtual TSharedRef<SWidget> GetColumnMenuContent(const FName InColumnId) override;
+	virtual void GetChildrenOfItem(const FNavigationToolItemPtr InItem, TArray<FNavigationToolItemPtr>& OutChildren) const override;
+	virtual void GetChildrenOfItem(const FNavigationToolItemPtr& InItem
+		, TArray<FNavigationToolItemPtr>& OutChildren
+		, ENavigationToolItemViewMode InViewMode
+		, const TSet<FNavigationToolItemPtr>& InRecursionDisallowedItems) const override;
+	virtual bool IsItemReadOnly(const FNavigationToolItemPtr& InItem) const override;
+	virtual bool CanSelectItem(const FNavigationToolItemPtr& InItem) const override;
+	virtual void SelectItems(TArray<FNavigationToolItemPtr> InItems, const ENavigationToolItemSelectionFlags InFlags = ENavigationToolItemSelectionFlags::SignalSelectionChange) override;
+	virtual void ClearItemSelection(bool bSignalSelectionChange = true) override;
+	virtual bool IsItemSelected(const FNavigationToolItemPtr& InItem) const override;
+	virtual TArray<FNavigationToolItemPtr> GetSelectedItems() const override;
+	virtual FReply OnDragDetected(const FGeometry& InGeometry, const FPointerEvent& InMouseEvent, const FNavigationToolItemPtr InTargetItem) override;
+	virtual FReply OnDrop(const FDragDropEvent& InDragDropEvent, EItemDropZone InDropZone, const FNavigationToolItemPtr InTargetItem) override;
+	virtual TOptional<EItemDropZone> OnCanDrop(const FDragDropEvent& InDragDropEvent
+		, EItemDropZone InDropZone
+		, const FNavigationToolItemPtr InTargetItem) const override;
+	virtual bool IsItemExpanded(const FNavigationToolItemPtr& InItem, const bool bInUseFilter = true) const override;
+	virtual void SetItemExpansion(const FNavigationToolItemPtr& InItem, const bool bInExpand, const bool bInUseFilter = true) override;
+	virtual void SetItemExpansionRecursive(const FNavigationToolItemPtr InItem, const bool bInExpand) override;
+	virtual void SetParentItemExpansions(const FNavigationToolItemPtr& InItem, const bool bInExpand) override;
+	virtual bool CanExpandAll() const override;
+	virtual void ExpandAll() override;
+	virtual void CollapseAll() override;
+	virtual bool CanCollapseAll() const override;
+	//~ End INavigationToolView
+
+	TSharedPtr<SWidget> CreateItemContextMenu();
+
+	/** Get the Columns currently allocated in this Navigation Tool Instance */
+	const TMap<FName, TSharedPtr<INavigationToolColumn>>& GetColumns() const { return Columns; }
+
+	/** Whether the given column should be shown by default. This is only used at the start when Creating the Columns */
+	bool ShouldShowColumnByDefault(const TSharedPtr<INavigationToolColumn>& InColumn) const;
+
+	/** Refreshes the Items visible in this View, and refreshes the Navigation Tool Widget if it was created */
+	void Refresh();
+
+	/** Refreshes the Items that will be at the Top Level of the Tree (e.g. Actors with no parent item) */
+	void UpdateRootVisibleItems();
+
+	/** Updates the Item Expansions in the Tree Widget based on whether the Items have the Expanded flag or not */
+	void UpdateItemExpansions();
+
+	/** Called when Object Replacement has taken place. Used to invalidate the widget for painting */
+	void NotifyObjectsReplaced();
+
+	/** The Root Item of the Tree. Convenience function to get the RootItem from FNavigationTool */
+	FNavigationToolItemPtr GetRootItem() const;
+
+	/** Gets the Top Level Items that should be visible in the Tree. This can differ between views depending on Filters and other factors */
+	const TArray<FNavigationToolItemPtr>& GetRootVisibleItems() const;
+
+	void SaveViewItemFlags(const FNavigationToolItemPtr& InItem, ENavigationToolItemFlags InFlags);
+	
+	ENavigationToolItemFlags GetViewItemFlags(const FNavigationToolItemPtr& InItem) const;
+
+	FLinearColor GetItemBrushColor(const FNavigationToolItemPtr InItem) const;
+
+	/** Gets the Currently Selected Item Count in the Tree View */
+	int32 GetViewSelectedItemCount() const;
+
+	/** Calculates the amount of Items that are visible in this Navigation Tool View */
+	int32 CalculateVisibleItemCount() const;
+
+	/** Whether Sync Selection is currently taking place */
+	bool IsSyncingItemSelection() const { return bSyncingItemSelection; }
+
+	TSharedPtr<FNavigationToolFilterBar> GetFilterBar() const { return FilterBar; }
+
+	/** Called when Item selection has changed. This can be called by the Navigation Tool Widget or by the View itself if there's no ToolWidget */
+	void NotifyItemSelectionChanged(const TArray<FNavigationToolItemPtr>& InSelectedItems
+		, const FNavigationToolItemPtr& InItem
+		, bool bUpdateModeTools);
+
+	/** Return whether the given item should be visible in this Navigation Tool Instance */
+	bool ShouldShowItem(const FNavigationToolItemPtr& InItem, bool bInUseFilters, ENavigationToolItemViewMode InViewMode) const;
+
+	/** Returns the Index of the Child from the Parent's Children Visible List (Visible List can vary between Tool Instances) */
+	int32 GetVisibleChildIndex(const FNavigationToolItemPtr& InParentItem, const FNavigationToolItemPtr& InChildItem) const;
+
+	/** Returns the Child from the Parent's Children Visible List at the given Index (Visible List can vary between Tool Instances) */
+	FNavigationToolItemPtr GetVisibleChildAt(const FNavigationToolItemPtr& InParentItem, int32 InChildIndex) const;
+
+	/** Returns whether the Tool is Locked. See FNavigationTool::IsToolLocked */
+	bool IsToolLocked() const;
+
+	/**
+	 * Activate an Item Column
+	 * @param InColumn the column to activate
+	 */
+	void ShowColumn(const TSharedPtr<INavigationToolColumn>& InColumn);
+
+	/**
+	 * Activate the Item Column with the given ID
+	 * @param InColumnId the filter id associated with the Filter to activate
+	 */
+	void ShowColumnById(const FName& InColumnId);
+
+	/**
+	 * Deactivates the given Item Column
+	 * @param InColumn the filter to remove from the active filter list
+	 */
+	void HideColumn(const TSharedPtr<INavigationToolColumn>& InColumn);
+
+	/** Whether the given Column is in the Active Column List */
+	virtual bool IsColumnVisible(const TSharedPtr<INavigationToolColumn>& InColumn) const override;
+
+	void ToggleViewModeSupport(ENavigationToolItemViewMode& InOutViewMode, ENavigationToolItemViewMode InFlags);
+	void ToggleItemDefaultViewModeSupport(ENavigationToolItemViewMode InFlags);
+	void ToggleItemProxyViewModeSupport(ENavigationToolItemViewMode InFlags);
+
+	ECheckBoxState GetViewModeCheckState(ENavigationToolItemViewMode InViewMode, ENavigationToolItemViewMode InFlags) const;
+	ECheckBoxState GetItemDefaultViewModeCheckState(ENavigationToolItemViewMode InFlags) const;
+	ECheckBoxState GetItemProxyViewModeCheckState(ENavigationToolItemViewMode InFlags) const;
+	
+	/**
+	 * Action to turn Muted Hierarchy on or off
+	 * @see bUseMutedHierarchy
+	 */
+	void ToggleMutedHierarchy();
+	bool CanToggleMutedHierarchy() const { return true; }
+	bool IsMutedHierarchyActive() const;
+
+	/**
+	 * Action to turn Auto Expand to Selection on or off
+	 * @see bAutoExpandToSelection
+	 */
+	void ToggleAutoExpandToSelection();
+	bool CanToggleAutoExpandToSelection() const { return true; }
+	bool ShouldAutoExpandToSelection() const;
+
+	void ToggleUseShortNames();
+	bool CanToggleUseShortNames() const { return true; }
+	bool ShouldUseShortNames() const;
+
+	void ToggleShowItemFilters();
+	bool CanToggleShowItemFilters() const { return true; }
+	bool ShouldShowItemFilters() const { return bShowItemFilters; }
+
+	void ToggleShowItemColumns();
+	bool CanToggleShowItemColumns() const { return true; }
+	bool ShouldShowItemColumns() const { return bShowItemColumns; }
+
+	/** Sets whether the given item type should be hidden or not */
+	void SetItemTypeHidden(const FName InItemTypeName, const bool bInHidden);
+	
+	/** Toggles the given Item Types to Hide/Show */
+	void ToggleHideItemTypes(const FName InItemTypeName);
+	ECheckBoxState GetToggleHideItemTypesState(const FName InItemTypeName) const;
+	
+	/** Hides the given Outliner Item Type from showing in this Outliner View (e.g. Hide Components is used) */
+	template<typename InItemType, typename = typename TEnableIf<TIsDerivedFrom<InItemType, INavigationToolItem>::IsDerived>::Type>
+	void HideItemType()
+	{
+		SetItemTypeHidden(TNavigationToolItemType<InItemType>::GetTypeId().ToName(), true);
+	}
+	
+	/** Shows the given Outliner Item Type if it was hidden in this Outliner View (e.g. Showing Components again) */
+	template<typename InItemType, typename = typename TEnableIf<TIsDerivedFrom<InItemType, INavigationToolItem>::IsDerived>::Type>
+	void ExposeItemType()
+	{
+		SetItemTypeHidden(TNavigationToolItemType<InItemType>::GetTypeId().ToName(), false);
+	}
+	
+	/** Whether the given Outliner Item Type is currently hidden in this Outliner View */
+	bool IsItemTypeHidden(const FName InItemTypeName) const;
+	bool IsItemTypeHidden(const FNavigationToolItemPtr& InItem) const;
+	
+	template<typename InItemType, typename = typename TEnableIf<TIsDerivedFrom<InItemType, INavigationToolItem>::IsDerived>::Type>
+	bool IsItemTypeHidden() const
+	{
+		return IsItemTypeHidden(InItemType::GetStaticTypeName());
+	}
+
+	/** Called when a Drag enters the Navigation Tool Widgets for the given Item (if null, treat as root) */
+	void OnDragEnter(const FDragDropEvent& InDragDropEvent, const FNavigationToolItemPtr InTargetItem);
+
+	/** Called when a Drag leaves the Navigation Tool Widgets for the given Item (if null, treat as root) */
+	void OnDragLeave(const FDragDropEvent& InDragDropEvent, const FNavigationToolItemPtr InTargetItem);
+
+	/** Called when there's a Drag being started/ended to the Navigation Tool Widget as a whole (i.e. root) and not an individual item */
+	void SetDragIntoTreeRoot(bool bInIsDraggingIntoTreeRoot);
+	
+	void RenameSelected();
+	void ResetRenaming();
+	void OnItemRenameAction(const ENavigationToolRenameAction InRenameAction, const TSharedPtr<INavigationToolView>& InToolView);
+	bool CanRenameSelected() const;
+
+	void DeleteSelected();
+	bool CanDeleteSelected() const;
+
+	void DuplicateSelected();
+	bool CanDuplicateSelected() const;
+
+	void SelectChildren(bool bIsRecursive);
+	bool CanSelectChildren() const;
+
+	void SelectParent();
+	bool CanSelectParent() const;
+
+	void SelectFirstChild();
+	bool CanSelectFirstChild() const;
+
+	void SelectSibling(int32 InDeltaIndex);
+	bool CanSelectSibling() const;
+
+	/** Called when the given item's expansion state (expanded/collapsed) has changed */
+	void OnItemExpansionChanged(const FNavigationToolItemPtr InItem, const bool bInIsExpanded);
+
+	void ScrollNextIntoView();
+	void ScrollPrevIntoView();
+	bool CanScrollNextIntoView() const;
+	void ScrollDeltaIndexIntoView(int32 DeltaIndex);
+	void ScrollItemIntoView(const FNavigationToolItemPtr& InItem);
+
+	/** Sorts the given Items in the user-defined order and selects them, and if the Widget is valid focusing to the item in the Widget */
+	void SortAndSelectItems(TArray<FNavigationToolItemPtr> InItemsToSelect);
+
+	/** Save the current state of the Navigation Tool */
+	void SaveViewState(const TSharedRef<FNavigationToolProvider>& InProvider);
+	/** Load the current state of the Navigation Tool */
+	void LoadViewState(const TSharedRef<FNavigationToolProvider>& InProvider);
+
+	/**  */
+	void SaveColumnState(const FName InColumnId = NAME_None);
+
+	void ResetColumnSize(const FName InColumnId);
+	bool CanResetColumnSize(const FName InColumnId) const;
+
+	void ResetVisibleColumnSizes();
+	bool CanResetAllColumnSizes() const;
+
+	/** Display a modal dialog for saving a new custom column view */
+	void SaveNewCustomColumnView();
+
+	/** Apply a saved custom column view by name */
+	void ApplyCustomColumnView(const FText InColumnViewName);
+
+protected:
+	/** Initializes the Navigation Tool Instance. Only executed once */
+	void Init(const TSharedRef<FNavigationTool>& InTool, const TSharedPtr<FUICommandList>& InBaseCommandList);
+
+	/** Sets the Selected Items to be the given Items in this List. */
+	void SetItemSelectionImpl(TArray<FNavigationToolItemPtr>&& InItems
+		, bool bSignalSelectionChange);
+
+	/** Triggers a Refresh on the FNavigationTool */
+	void RefreshTool(bool bInImmediateRefresh);
+
+	void EnsureToolViewCount(const int32 InToolViewId) const;
+
+	/**  */
+	void SaveFilterState(FNavigationToolViewSaveState& OutViewSaveState);
+	/**  */
+	void LoadFilterState(const FNavigationToolViewSaveState& InViewSaveState
+		, const bool bInDisableAllFilters = true
+		, const bool bInRequestFilterUpdate = true);
+
+	/** Save the Item State in the Navigation Tool Widget (e.g. Item Scoped Flags) */
+	void SaveToolViewItems(FNavigationToolViewSaveState& OutViewSaveState);
+	/** Load the Item State in the Navigation Tool Widget (e.g. Item Scoped Flags) */
+	void LoadToolViewItems(FNavigationToolViewSaveState& InViewSaveState);
+
+	bool CanFocusSingleSelection() const;
+	void FocusSingleSelection();
+
+	bool CanFocusInContentBrowser() const;
+	void FocusInContentBrowser();
+
+	bool UpdateFilters();
+
+	/** Local Identifier of this Instance */
+	int32 ToolViewId = -1;
+
+	/** Weak pointer to the outliner this is a view of */
+	TWeakPtr<FNavigationTool> WeakTool;
+
+	/** Widget showing the View. Can be null if instanced for testing (i.e. bCreateToolWidget is false in Init) */
+	TSharedPtr<SNavigationToolView> ToolViewWidget;
+
+	/** Command List mapped to the View to handle things like Selected Items */
+	TSharedPtr<FUICommandList> ViewCommandList;
+
+	/**  */
+	const TObjectPtr<UNavigationToolSettings> ToolSettings;
+
+	/** Root Items from Navigation Tool visible to this Instance. i.e. a Subset of all the Items in FNavigationTool */
+	TArray<FNavigationToolItemPtr> RootVisibleItems;
+
+	/** A list of the Selected Items in this Navigation Tool View */
+	TArray<FNavigationToolItemPtr> SelectedItems;
+
+	/** Set of Items that are Currently Read Only in this Instance. Mutable as this changes on GetChildrenOfItem const func */
+	mutable TSet<FNavigationToolItemPtr> ReadOnlyItems;
+
+	/** A set of the Navigation Tool Item Type Names that should be hidden */
+	TSet<FName> HiddenItemTypes;
+
+	TSharedPtr<FNavigationToolFilterBar> FilterBar;
+
+	/** Delegate called at the end of the Navigation Tool View Refresh */
+	FOnToolViewRefreshed OnToolViewRefreshed;
+
+	/** Map of the Column Ids and the Column with that Id. This is filled in FNavigationToolView::CreateColumns */
+	TMap<FName, TSharedPtr<INavigationToolColumn>> Columns;
+
+	/** Right click context menu for items */
+	TSharedPtr<FNavigationToolItemContextMenu> ItemContextMenu;
+
+	/** The Index of the Item that will be scrolled into View */
+	int32 NextSelectedItemIntoView = -1;
+
+	/** Selected Items Sorted from Top To Bottom */
+	TArray<FNavigationToolItemPtr> SortedSelectedItems;
+
+	/** The list of items that need renaming and are waiting for their turn to be renamed */
+	TArray<FNavigationToolItemPtr> ItemsRemainingRename;
+	
+	/** The current item in the process of renaming. Null if no renaming is taking place */
+	FNavigationToolItemPtr CurrentItemRenaming;
+	
+	/** Flag to call FNavigationToolView::Refresh next tick */
+	bool bRefreshRequested = false;
+	
+	/** Whether Item renaming is taking place in this View */
+	bool bRenamingItems = false;
+	
+	/** Whether Renaming Items should be processed next tick */
+	bool bRequestedRename = false;
+
+	/** Whether Items are currently being synced to what the Widget or Interface sent. Used as a re-enter guard */
+	bool bSyncingItemSelection = false;
+	
+	/** Flag used for the Navigation Tool Widget to determine if the Item Filter Bar should open and show the Item Filters */
+	bool bShowItemFilters = false;
+
+	bool bShowItemColumns = false;
+
+	bool bFilterUpdateRequested = false;
+};
+
+}
diff --git a/Engine/Plugins/Experimental/SequenceNavigator/Source/SequenceNavigator/Private/Widgets/Columns/SNavigationToolDeactiveState.cpp b/Engine/Plugins/Experimental/SequenceNavigator/Source/SequenceNavigator/Private/Widgets/Columns/SNavigationToolDeactiveState.cpp
new file mode 100644
--- /dev/null
+++ b/Engine/Plugins/Experimental/SequenceNavigator/Source/SequenceNavigator/Private/Widgets/Columns/SNavigationToolDeactiveState.cpp
@@ -0,0 +1,237 @@
+// Copyright Epic Games, Inc. All Rights Reserved.
+
+#include "SNavigationToolDeactiveState.h"
+#include "Columns/NavigationToolDeactiveStateColumn.h"
+#include "Editor.h"
+#include "Extensions/ISequenceInactivableExtension.h"
+#include "Items/INavigationToolItem.h"
+#include "NavigationToolView.h"
+#include "Widgets/SNavigationToolTreeRow.h"
+
+#define LOCTEXT_NAMESPACE "SNavigationToolDeactiveState"
+
+namespace UE::Sequencer::NavigationTool
+{
+
+class FEvaluationStateDragDropOp : public FDragDropOperation
+{
+public:
+	DRAG_DROP_OPERATOR_TYPE(FVisibilityDragDropOp, FDragDropOperation)
+
+	/** Flag which defines whether to hide destination items or not */
+	EItemSequenceInactiveState InactiveState;
+
+	/** Undo transaction stolen from the gutter which is kept alive for the duration of the drag */
+	TUniquePtr<FScopedTransaction> UndoTransaction;
+
+	/** The widget decorator to use */
+	virtual TSharedPtr<SWidget> GetDefaultDecorator() const override
+	{
+		return SNullWidget::NullWidget;
+	}
+
+	/** Create a new drag and drop operation out of the specified flag */
+	static TSharedRef<FEvaluationStateDragDropOp> New(const EItemSequenceInactiveState InEvaluationState, TUniquePtr<FScopedTransaction>& ScopedTransaction)
+	{
+		TSharedRef<FEvaluationStateDragDropOp> Operation = MakeShared<FEvaluationStateDragDropOp>();
+		Operation->InactiveState = InEvaluationState;
+		Operation->UndoTransaction = MoveTemp(ScopedTransaction);
+		Operation->Construct();
+		return Operation;
+	}
+};
+
+EItemSequenceInactiveState GetItemInactiveState(const FNavigationToolItemPtr& InItem)
+{
+	if (!InItem.IsValid())
+	{
+		return EItemSequenceInactiveState::None;
+	}
+
+	ISequenceInactivableExtension* const InactivatableItem = InItem->CastTo<ISequenceInactivableExtension>();
+	if (!InactivatableItem)
+	{
+		return EItemSequenceInactiveState::None;
+	}
+
+	return InactivatableItem->GetInactiveState();
+}
+
+void SetItemInactive(const FNavigationToolItemPtr& InItem, const bool bInInactive)
+{
+	if (!InItem.IsValid())
+	{
+		return;
+	}
+
+	if (ISequenceInactivableExtension* const InactivatableItem = InItem->CastTo<ISequenceInactivableExtension>())
+	{
+		InactivatableItem->SetIsInactive(bInInactive);
+	}
+}
+
+void SNavigationToolDeactiveState::Construct(const FArguments& InArgs
+	, const TSharedRef<FNavigationToolDeactiveStateColumn>& InColumn
+	, const FNavigationToolItemRef& InItem
+	, const TSharedRef<INavigationToolView>& InView
+	, const TSharedRef<SNavigationToolTreeRow>& InRowWidget)
+{
+	WeakColumn = InColumn;
+	WeakItem = InItem;
+	WeakView = InView;
+	WeakRowWidget = InRowWidget;
+
+	SImage::Construct(SImage::FArguments()
+		.IsEnabled(this, &SNavigationToolDeactiveState::IsVisibilityWidgetEnabled)
+		.ColorAndOpacity(this, &SNavigationToolDeactiveState::GetForegroundColor)
+		.Image(this, &SNavigationToolDeactiveState::GetBrush));
+}
+
+FReply SNavigationToolDeactiveState::OnDragDetected(const FGeometry& MyGeometry, const FPointerEvent& MouseEvent)
+{
+	if (MouseEvent.IsMouseButtonDown(EKeys::LeftMouseButton))
+	{
+		const EItemSequenceInactiveState NewState = (GetInactiveState() != EItemSequenceInactiveState::None)
+			? EItemSequenceInactiveState::Inactive : EItemSequenceInactiveState::None;
+		return FReply::Handled().BeginDragDrop(FEvaluationStateDragDropOp::New(NewState, UndoTransaction));
+	}
+	return FReply::Unhandled();
+}
+
+void SNavigationToolDeactiveState::OnDragEnter(const FGeometry& MyGeometry, const FDragDropEvent& DragDropEvent)
+{
+	const TSharedPtr<FEvaluationStateDragDropOp> DragDropOp = DragDropEvent.GetOperationAs<FEvaluationStateDragDropOp>();
+	if (DragDropOp.IsValid())
+	{
+		const bool bItemInactive = (DragDropOp->InactiveState != EItemSequenceInactiveState::None);
+		SetIsInactive(bItemInactive);
+	}
+}
+
+FReply SNavigationToolDeactiveState::HandleClick()
+{
+	if (!IsVisibilityWidgetEnabled())
+	{
+		return FReply::Unhandled();
+	}
+
+	const TSharedPtr<INavigationToolView> ToolView = WeakView.Pin();
+	const FNavigationToolItemPtr Item = WeakItem.Pin();
+	const TSharedPtr<FNavigationToolDeactiveStateColumn> Column = WeakColumn.Pin();
+
+	if (!ToolView.IsValid() || !Item.IsValid() || !Column.IsValid())
+	{
+		return FReply::Unhandled();
+	}
+
+	// Open an undo transaction
+	UndoTransaction = MakeUnique<FScopedTransaction>(LOCTEXT("SetNavigationToolItemInactiveState", "Set Item Inactive State"));
+
+	const bool bNewIsInactive = (GetInactiveState() != EItemSequenceInactiveState::Inactive);
+
+	// We operate on all the selected items if the specified item is selected
+	if (ToolView->IsItemSelected(Item.ToSharedRef()))
+	{
+		for (FNavigationToolItemPtr& SelectedItem : ToolView->GetSelectedItems())
+		{
+			SetItemInactive(SelectedItem, bNewIsInactive);
+		}
+	}
+	else
+	{
+		SetIsInactive(bNewIsInactive);
+	}
+
+	return FReply::Handled().DetectDrag(SharedThis(this), EKeys::LeftMouseButton);
+}
+
+FReply SNavigationToolDeactiveState::OnMouseButtonDoubleClick(const FGeometry& InMyGeometry, const FPointerEvent& InMouseEvent)
+{
+	return HandleClick();
+}
+
+FReply SNavigationToolDeactiveState::OnMouseButtonDown(const FGeometry& MyGeometry, const FPointerEvent& MouseEvent)
+{
+	if (MouseEvent.GetEffectingButton() != EKeys::LeftMouseButton)
+	{
+		return FReply::Unhandled();
+	}
+
+	return HandleClick();
+}
+
+FReply SNavigationToolDeactiveState::OnMouseButtonUp(const FGeometry& MyGeometry, const FPointerEvent& MouseEvent)
+{
+	if (MouseEvent.GetEffectingButton() == EKeys::LeftMouseButton)
+	{
+		UndoTransaction.Reset();
+		return FReply::Handled();
+	}
+
+	return FReply::Unhandled();
+}
+
+void SNavigationToolDeactiveState::OnMouseCaptureLost(const FCaptureLostEvent& CaptureLostEvent)
+{
+	UndoTransaction.Reset();
+}
+
+const FSlateBrush* SNavigationToolDeactiveState::GetBrush() const
+{
+	return FAppStyle::GetBrush(TEXT("Sequencer.Column.Mute"));
+}
+
+FSlateColor SNavigationToolDeactiveState::GetForegroundColor() const
+{
+	const FNavigationToolItemPtr Item = WeakItem.Pin();
+	if (!Item.IsValid())
+	{
+		return FStyleColors::Transparent;
+	}
+
+	const bool bIsItemHovered = WeakRowWidget.IsValid() && WeakRowWidget.Pin()->IsHovered();
+	if (bIsItemHovered || IsHovered())
+	{
+		switch (GetInactiveState())
+		{
+		case EItemSequenceInactiveState::None:
+			return FStyleColors::White25;
+
+		case EItemSequenceInactiveState::PartiallyInactive:
+			return FStyleColors::ForegroundHover;
+
+		case EItemSequenceInactiveState::Inactive:
+			return FStyleColors::ForegroundHover;
+		}
+	}
+	else
+	{
+		switch (GetInactiveState())
+		{
+		case EItemSequenceInactiveState::None:
+			return FStyleColors::Transparent;
+
+		case EItemSequenceInactiveState::PartiallyInactive:
+			return FStyleColors::White25;
+
+		case EItemSequenceInactiveState::Inactive:
+			return FStyleColors::Foreground;
+		}
+	}
+
+	return FStyleColors::Transparent;
+}
+
+EItemSequenceInactiveState SNavigationToolDeactiveState::GetInactiveState() const
+{
+	return GetItemInactiveState(WeakItem.Pin());
+}
+
+void SNavigationToolDeactiveState::SetIsInactive(const bool bInIsInactive)
+{
+	SetItemInactive(WeakItem.Pin(), bInIsInactive);
+}
+
+}
+
+#undef LOCTEXT_NAMESPACE
diff --git a/Engine/Plugins/Experimental/SequenceNavigator/Source/SequenceNavigator/Private/Widgets/Columns/SNavigationToolDeactiveState.h b/Engine/Plugins/Experimental/SequenceNavigator/Source/SequenceNavigator/Private/Widgets/Columns/SNavigationToolDeactiveState.h
new file mode 100644
--- /dev/null
+++ b/Engine/Plugins/Experimental/SequenceNavigator/Source/SequenceNavigator/Private/Widgets/Columns/SNavigationToolDeactiveState.h
@@ -0,0 +1,65 @@
+// Copyright Epic Games, Inc. All Rights Reserved.
+
+#pragma once
+
+#include "NavigationToolDefines.h"
+#include "ScopedTransaction.h"
+#include "Templates/UniquePtr.h"
+#include "Widgets/Images/SImage.h"
+
+namespace UE::Sequencer::NavigationTool
+{
+
+class FNavigationToolItem;
+class FNavigationToolDeactiveStateColumn;
+class INavigationToolView;
+class SNavigationToolTreeRow;
+enum class EItemSequenceInactiveState;
+
+/** Widget responsible for managing the visibility for a single item */
+class SNavigationToolDeactiveState : public SImage
+{
+public:
+	SLATE_BEGIN_ARGS(SNavigationToolDeactiveState) {}
+	SLATE_END_ARGS()
+
+	void Construct(const FArguments& InArgs
+		, const TSharedRef<FNavigationToolDeactiveStateColumn>& InColumn
+		, const FNavigationToolItemRef& InItem
+		, const TSharedRef<INavigationToolView>& InView
+		, const TSharedRef<SNavigationToolTreeRow>& InRowWidget);
+
+private:
+	/** Returns whether the widget is enabled or not */
+	virtual bool IsVisibilityWidgetEnabled() const { return true; }
+
+	virtual const FSlateBrush* GetBrush() const;
+
+	//~ Begin SWidget
+	virtual FReply OnDragDetected(const FGeometry& MyGeometry, const FPointerEvent& MouseEvent) override;
+	virtual void OnDragEnter(const FGeometry& MyGeometry, const FDragDropEvent& DragDropEvent) override;
+	virtual FReply OnMouseButtonDoubleClick(const FGeometry& InMyGeometry, const FPointerEvent& InMouseEvent) override;
+	virtual FReply OnMouseButtonDown(const FGeometry& MyGeometry, const FPointerEvent& MouseEvent) override;
+	virtual FReply OnMouseButtonUp(const FGeometry& MyGeometry, const FPointerEvent& MouseEvent) override;
+	virtual void OnMouseCaptureLost(const FCaptureLostEvent& CaptureLostEvent) override;
+	//~ End SWidget
+
+	FReply HandleClick();
+
+	virtual FSlateColor GetForegroundColor() const override;
+
+	EItemSequenceInactiveState GetInactiveState() const;
+	void SetIsInactive(const bool bInIsInactive);
+
+	TWeakPtr<FNavigationToolDeactiveStateColumn> WeakColumn;
+
+	FNavigationToolItemWeakPtr WeakItem;
+
+	TWeakPtr<INavigationToolView> WeakView;
+
+	TWeakPtr<SNavigationToolTreeRow> WeakRowWidget;
+
+	TUniquePtr<FScopedTransaction> UndoTransaction;
+};
+
+}
