[Sequencer] Move Navigation Tool to new experimental SequenceNavigator plugin

diff --git a/Engine/Plugins/Experimental/SequenceNavigator/Source/SequenceNavigator/Private/Widgets/Columns/SNavigationToolInTime.cpp b/Engine/Plugins/Experimental/SequenceNavigator/Source/SequenceNavigator/Private/Widgets/Columns/SNavigationToolInTime.cpp
new file mode 100644
--- /dev/null
+++ b/Engine/Plugins/Experimental/SequenceNavigator/Source/SequenceNavigator/Private/Widgets/Columns/SNavigationToolInTime.cpp
@@ -0,0 +1,72 @@
+// Copyright Epic Games, Inc. All Rights Reserved.
+
+#include "SNavigationToolInTime.h"
+#include "ScopedTransaction.h"
+#include "Extensions/IInTimeExtension.h"
+#include "Items/NavigationToolItemUtils.h"
+#include "NavigationToolView.h"
+#include "Widgets/SNavigationToolTreeRow.h"
+
+#define LOCTEXT_NAMESPACE "SNavigationToolInTime"
+
+namespace UE::Sequencer::NavigationTool
+{
+
+void SNavigationToolInTime::Construct(const FArguments& InArgs
+	, const FNavigationToolItemRef& InItem
+	, const TSharedRef<INavigationToolView>& InView
+	, const TSharedRef<SNavigationToolTreeRow>& InRowWidget)
+{
+	WeakItem = InItem;
+	WeakView = InView;
+	WeakRowWidget = InRowWidget;
+
+	if (!InItem->IsA<IInTimeExtension>())
+	{
+		return;
+	}
+
+	SNavigationToolTime::Construct(SNavigationToolTime::FArguments(), InItem, InView, InRowWidget);
+}
+
+IInTimeExtension* SNavigationToolInTime::GetInTimeExtension() const
+{
+	return WeakItem.IsValid() ? WeakItem.Pin()->CastTo<IInTimeExtension>() : nullptr;
+}
+
+FName SNavigationToolInTime::GetStyleName() const
+{
+	return TEXT("SpinBox.InTime");
+}
+
+double SNavigationToolInTime::GetFrameTimeValue() const
+{
+	if (const IInTimeExtension* const InTimeExtension = GetInTimeExtension())
+	{
+		return InTimeExtension->GetInTime().Value;
+	}
+	return 0;
+}
+
+void SNavigationToolInTime::OnFrameTimeValueCommitted(const double InNewValue, const ETextCommit::Type InCommitType)
+{
+	IInTimeExtension* const InTimeExtension = GetInTimeExtension();
+	if (!InTimeExtension)
+	{
+		return;
+	}
+
+	const bool bShouldTransact = !UndoTransaction.IsValid() && (InCommitType == ETextCommit::OnEnter);
+	const FScopedTransaction Transaction(GetTransactionText(), bShouldTransact);
+
+	return InTimeExtension->SetInTime(static_cast<int32>(InNewValue));
+}
+
+FText SNavigationToolInTime::GetTransactionText() const
+{
+	return LOCTEXT("SetInTimeTransaction", "Set In Time");
+}
+
+}
+
+#undef LOCTEXT_NAMESPACE
diff --git a/Engine/Plugins/Experimental/SequenceNavigator/Source/SequenceNavigator/Private/Widgets/Columns/SNavigationToolLock.cpp b/Engine/Plugins/Experimental/SequenceNavigator/Source/SequenceNavigator/Private/Widgets/Columns/SNavigationToolLock.cpp
new file mode 100644
--- /dev/null
+++ b/Engine/Plugins/Experimental/SequenceNavigator/Source/SequenceNavigator/Private/Widgets/Columns/SNavigationToolLock.cpp
@@ -0,0 +1,238 @@
+// Copyright Epic Games, Inc. All Rights Reserved.
+
+#include "SNavigationToolLock.h"
+#include "Items/INavigationToolItem.h"
+#include "Items/NavigationToolItemUtils.h"
+#include "NavigationToolSettings.h"
+#include "NavigationToolView.h"
+#include "Widgets/SNavigationToolTreeRow.h"
+
+#define LOCTEXT_NAMESPACE "SNavigationToolLock"
+
+namespace UE::Sequencer::NavigationTool
+{
+
+class FLockDragDropOp : public FDragDropOperation
+{
+public:
+	DRAG_DROP_OPERATOR_TYPE(FLockDragDropOp, FDragDropOperation)
+
+	/** Flag which defines whether to lock destination items or not */
+	bool bShouldLock;
+
+	/** Undo transaction stolen from the gutter which is kept alive for the duration of the drag */
+	TUniquePtr<FScopedTransaction> UndoTransaction;
+
+	/** The widget decorator to use */
+	virtual TSharedPtr<SWidget> GetDefaultDecorator() const override
+	{
+		return SNullWidget::NullWidget;
+	}
+
+	/** Create a new drag and drop operation out of the specified flag */
+	static TSharedRef<FLockDragDropOp> New(const bool bShouldLock, TUniquePtr<FScopedTransaction>& ScopedTransaction)
+	{
+		const TSharedRef<FLockDragDropOp> Operation = MakeShared<FLockDragDropOp>();
+		Operation->bShouldLock     = bShouldLock;
+		Operation->UndoTransaction = MoveTemp(ScopedTransaction);
+		Operation->Construct();
+		return Operation;
+	}
+};
+
+EItemSequenceLockState GetItemLockState(const FNavigationToolItemPtr& InItem)
+{
+	if (!InItem.IsValid())
+	{
+		return EItemSequenceLockState::None;
+	}
+
+	ISequenceLockableExtension* const LockableItem = InItem->CastTo<ISequenceLockableExtension>();
+	if (!LockableItem)
+	{
+		return EItemSequenceLockState::None;
+	}
+
+	return LockableItem->GetLockState();
+}
+
+void SetItemLocked(const FNavigationToolItemPtr& InItem, const bool bInIsLocked)
+{
+	if (!InItem.IsValid())
+	{
+		return;
+	}
+
+	if (ISequenceLockableExtension* const LockableItem = InItem->CastTo<ISequenceLockableExtension>())
+	{
+		LockableItem->SetIsLocked(bInIsLocked);
+	}
+}
+
+void SNavigationToolLock::Construct(const FArguments& InArgs
+	, const FNavigationToolItemRef& InItem
+	, const TSharedRef<INavigationToolView>& InView
+	, const TSharedRef<SNavigationToolTreeRow>& InRowWidget)
+{
+	WeakItem = InItem;
+	WeakView = InView;
+	WeakRowWidget = InRowWidget;
+
+	SImage::Construct(SImage::FArguments()
+		.ColorAndOpacity(this, &SNavigationToolLock::GetForegroundColor)
+		.Image(this, &SNavigationToolLock::GetBrush));
+}
+
+FSlateColor SNavigationToolLock::GetForegroundColor() const
+{
+	const FNavigationToolItemPtr Item = WeakItem.Pin();
+	if (!Item.IsValid())
+	{
+		return FSlateColor::UseForeground();
+	}
+
+	const bool bIsItemSelected = WeakView.IsValid() && WeakView.Pin()->IsItemSelected(Item.ToSharedRef());
+	const bool bIsItemHovered = WeakRowWidget.IsValid() && WeakRowWidget.Pin()->IsHovered();
+	const bool bAlwaysShowLock = GetDefault<UNavigationToolSettings>()->ShouldAlwaysShowLockState();
+
+	// We can hide the brush if settings for "Always Showing Lock State" is OFF
+	// and item is not locked while also not being selected or hovered
+	if (!bAlwaysShowLock
+		&& GetLockState() == EItemSequenceLockState::None
+		&& !bIsItemSelected
+		&& !bIsItemHovered)
+	{
+		return FLinearColor::Transparent;
+	}
+
+	if (bIsItemHovered || IsHovered())
+	{
+		switch (GetLockState())
+		{
+		case EItemSequenceLockState::None:
+			return FStyleColors::White25;
+
+		case EItemSequenceLockState::PartiallyLocked:
+			return FStyleColors::ForegroundHover;
+
+		case EItemSequenceLockState::Locked:
+			return FStyleColors::ForegroundHover;
+		}
+	}
+	else
+	{
+		switch (GetLockState())
+		{
+		case EItemSequenceLockState::None:
+			return FStyleColors::Transparent;
+
+		case EItemSequenceLockState::PartiallyLocked:
+			return FStyleColors::White25;
+
+		case EItemSequenceLockState::Locked:
+			return FStyleColors::Foreground;
+		}
+	}
+
+	return FSlateColor::UseForeground();
+}
+
+const FSlateBrush* SNavigationToolLock::GetBrush() const
+{
+	return (GetLockState() == EItemSequenceLockState::None)
+		? FAppStyle::GetBrush(TEXT("Icons.Unlock"))
+		: FAppStyle::GetBrush(TEXT("Icons.Lock"));
+}
+
+FReply SNavigationToolLock::OnDragDetected(const FGeometry& MyGeometry, const FPointerEvent& MouseEvent)
+{
+	if (MouseEvent.IsMouseButtonDown(EKeys::LeftMouseButton))
+	{
+		const bool bShouldLock = GetLockState() == EItemSequenceLockState::Locked;
+		return FReply::Handled().BeginDragDrop(FLockDragDropOp::New(bShouldLock, UndoTransaction));
+	}
+	return FReply::Unhandled();
+}
+
+void SNavigationToolLock::OnDragEnter(const FGeometry& MyGeometry, const FDragDropEvent& DragDropEvent)
+{
+	if (const TSharedPtr<FLockDragDropOp> DragDropOp = DragDropEvent.GetOperationAs<FLockDragDropOp>())
+	{
+		SetIsLocked(DragDropOp->bShouldLock);
+	}
+}
+
+FReply SNavigationToolLock::HandleClick()
+{
+	const TSharedPtr<INavigationToolView> ToolView = WeakView.Pin();
+	const FNavigationToolItemPtr Item = WeakItem.Pin();
+
+	if (!ToolView.IsValid() || !Item.IsValid())
+	{
+		return FReply::Unhandled();
+	}
+
+	UndoTransaction = MakeUnique<FScopedTransaction>(LOCTEXT("SetNavigationToolItemLock", "Set Item Lock"));
+
+	const bool bNewIsLocked = (GetLockState() != EItemSequenceLockState::Locked);
+
+	// We operate on all the selected items if the specified item is selected
+	if (ToolView->IsItemSelected(Item.ToSharedRef()))
+	{
+		for (FNavigationToolItemPtr& SelectedItem : ToolView->GetSelectedItems())
+		{
+			SetItemLocked(SelectedItem, bNewIsLocked);
+		}
+	}
+	else
+	{
+		SetIsLocked(bNewIsLocked);
+	}
+
+	return FReply::Handled().DetectDrag(SharedThis(this), EKeys::LeftMouseButton);
+}
+
+FReply SNavigationToolLock::OnMouseButtonDoubleClick(const FGeometry& InMyGeometry, const FPointerEvent& InMouseEvent)
+{
+	return HandleClick();
+}
+
+FReply SNavigationToolLock::OnMouseButtonDown(const FGeometry& MyGeometry, const FPointerEvent& MouseEvent)
+{
+	if (MouseEvent.GetEffectingButton() != EKeys::LeftMouseButton)
+	{
+		return FReply::Unhandled();
+	}
+	return HandleClick();
+}
+
+FReply SNavigationToolLock::OnMouseButtonUp(const FGeometry& MyGeometry, const FPointerEvent& MouseEvent)
+{
+	if (MouseEvent.GetEffectingButton() == EKeys::LeftMouseButton)
+	{
+		UndoTransaction.Reset();
+
+		return FReply::Handled();
+	}
+
+	return FReply::Unhandled();
+}
+
+void SNavigationToolLock::OnMouseCaptureLost(const FCaptureLostEvent& CaptureLostEvent)
+{
+	UndoTransaction.Reset();
+}
+
+EItemSequenceLockState SNavigationToolLock::GetLockState() const
+{
+	return GetItemLockState(WeakItem.Pin());
+}
+
+void SNavigationToolLock::SetIsLocked(const bool bInIsLocked)
+{
+	SetItemLocked(WeakItem.Pin(), bInIsLocked);
+}
+
+}
+
+#undef LOCTEXT_NAMESPACE
diff --git a/Engine/Plugins/Experimental/SequenceNavigator/Source/SequenceNavigator/Private/Widgets/Columns/SNavigationToolLock.h b/Engine/Plugins/Experimental/SequenceNavigator/Source/SequenceNavigator/Private/Widgets/Columns/SNavigationToolLock.h
new file mode 100644
--- /dev/null
+++ b/Engine/Plugins/Experimental/SequenceNavigator/Source/SequenceNavigator/Private/Widgets/Columns/SNavigationToolLock.h
@@ -0,0 +1,73 @@
+// Copyright Epic Games, Inc. All Rights Reserved.
+
+#pragma once
+
+#include "Extensions/ISequenceLockableExtension.h"
+#include "NavigationToolDefines.h"
+#include "ScopedTransaction.h"
+#include "Templates/UniquePtr.h"
+#include "Widgets/Images/SImage.h"
+
+class FDragDropEvent;
+class FReply;
+struct FCaptureLostEvent;
+struct FGeometry;
+struct FPointerEvent;
+struct FSlateBrush;
+struct FSlateColor;
+
+namespace UE::Sequencer::NavigationTool
+{
+
+class INavigationToolView;
+class SNavigationToolTreeRow;
+
+class SNavigationToolLock : public SImage
+{
+public:
+	SLATE_BEGIN_ARGS(SNavigationToolLock) {}
+	SLATE_END_ARGS()
+
+	virtual ~SNavigationToolLock() override {}
+
+	void Construct(const FArguments& InArgs
+		, const FNavigationToolItemRef& InItem
+		, const TSharedRef<INavigationToolView>& InView
+		, const TSharedRef<SNavigationToolTreeRow>& InRowWidget);
+
+	virtual FSlateColor GetForegroundColor() const override;
+
+	EItemSequenceLockState GetLockState() const;
+
+	void SetIsLocked(const bool bInIsLocked);
+
+	const FSlateBrush* GetBrush() const;
+
+	virtual FReply OnDragDetected(const FGeometry& MyGeometry, const FPointerEvent& MouseEvent) override;
+
+	/** If a lock drag drop operation has entered this widget, set its item to the new lock state */
+	virtual void OnDragEnter(const FGeometry& MyGeometry, const FDragDropEvent& DragDropEvent) override;
+
+	FReply HandleClick();
+
+	virtual FReply OnMouseButtonDoubleClick(const FGeometry& InMyGeometry, const FPointerEvent& InMouseEvent) override;
+
+	virtual FReply OnMouseButtonDown(const FGeometry& MyGeometry, const FPointerEvent& MouseEvent) override;
+
+	virtual FReply OnMouseButtonUp(const FGeometry& MyGeometry, const FPointerEvent& MouseEvent) override;
+
+	virtual void OnMouseCaptureLost(const FCaptureLostEvent& CaptureLostEvent) override;
+
+	//void OnSetItemLockedState(const FNavigationToolItemPtr& InItem, const bool bInIsLocked);
+
+private:
+	FNavigationToolItemWeakPtr WeakItem;
+
+	TWeakPtr<INavigationToolView> WeakView;
+
+	TWeakPtr<SNavigationToolTreeRow> WeakRowWidget;
+
+	TUniquePtr<FScopedTransaction> UndoTransaction;
+};
+
+}
diff --git a/Engine/Plugins/Experimental/SequenceNavigator/Source/SequenceNavigator/Private/Widgets/Columns/SNavigationToolMarkerVisibility.cpp b/Engine/Plugins/Experimental/SequenceNavigator/Source/SequenceNavigator/Private/Widgets/Columns/SNavigationToolMarkerVisibility.cpp
new file mode 100644
--- /dev/null
+++ b/Engine/Plugins/Experimental/SequenceNavigator/Source/SequenceNavigator/Private/Widgets/Columns/SNavigationToolMarkerVisibility.cpp
@@ -0,0 +1,239 @@
+// Copyright Epic Games, Inc. All Rights Reserved.
+
+#include "SNavigationToolMarkerVisibility.h"
+#include "Columns/NavigationToolMarkerVisibilityColumn.h"
+#include "Items/INavigationToolItem.h"
+#include "Items/NavigationToolItemUtils.h"
+#include "NavigationToolView.h"
+#include "Styling/StyleColors.h"
+#include "Widgets/SNavigationToolTreeRow.h"
+
+#define LOCTEXT_NAMESPACE "SNavigationToolMarkerVisibility"
+
+namespace UE::Sequencer::NavigationTool
+{
+
+class FMarkerVisibilityDragDropOp : public FDragDropOperation
+{
+public:
+	DRAG_DROP_OPERATOR_TYPE(FMarkerVisibilityDragDropOp, FDragDropOperation)
+
+	/** Flag which defines whether to hide destination items or not */
+	EItemMarkerVisibility MarkerVisibility;
+
+	/** Undo transaction stolen from the gutter which is kept alive for the duration of the drag */
+	TUniquePtr<FScopedTransaction> UndoTransaction;
+
+	/** The widget decorator to use */
+	virtual TSharedPtr<SWidget> GetDefaultDecorator() const override
+	{
+		return SNullWidget::NullWidget;
+	}
+
+	/** Create a new drag and drop operation out of the specified flag */
+	static TSharedRef<FMarkerVisibilityDragDropOp> New(const EItemMarkerVisibility InMarkerVisibility, TUniquePtr<FScopedTransaction>& ScopedTransaction)
+	{
+		TSharedRef<FMarkerVisibilityDragDropOp> Operation = MakeShared<FMarkerVisibilityDragDropOp>();
+		Operation->MarkerVisibility = InMarkerVisibility;
+		Operation->UndoTransaction = MoveTemp(ScopedTransaction);
+		Operation->Construct();
+		return Operation;
+	}
+};
+
+EItemMarkerVisibility GetItemMarkerVisibility(const FNavigationToolItemPtr& InItem)
+{
+	if (!InItem.IsValid())
+	{
+		return EItemMarkerVisibility::None;
+	}
+
+	IMarkerVisibilityExtension* const MarkerVisibilityItem = InItem->CastTo<IMarkerVisibilityExtension>();
+	if (!MarkerVisibilityItem)
+	{
+		return EItemMarkerVisibility::None;
+	}
+
+	return MarkerVisibilityItem->GetMarkerVisibility();
+}
+
+void SetItemMarkerVisibility(const FNavigationToolItemPtr& InItem, const bool bInMarkersVisible)
+{
+	if (!InItem.IsValid())
+	{
+		return;
+	}
+
+	if (IMarkerVisibilityExtension* const MarkerVisibilityItem = InItem->CastTo<IMarkerVisibilityExtension>())
+	{
+		MarkerVisibilityItem->SetMarkerVisibility(bInMarkersVisible);
+	}
+}
+
+void SNavigationToolMarkerVisibility::Construct(const FArguments& InArgs
+	, const TSharedRef<FNavigationToolMarkerVisibilityColumn>& InColumn
+	, const FNavigationToolItemRef& InItem
+	, const TSharedRef<INavigationToolView>& InView
+	, const TSharedRef<SNavigationToolTreeRow>& InRowWidget)
+{
+	WeakColumn = InColumn;
+	WeakItem = InItem;
+	WeakView = InView;
+	WeakRowWidget = InRowWidget;
+
+	SImage::Construct(SImage::FArguments()
+		.IsEnabled(this, &SNavigationToolMarkerVisibility::IsVisibilityWidgetEnabled)
+		.ColorAndOpacity(this, &SNavigationToolMarkerVisibility::GetForegroundColor)
+		.Image(this, &SNavigationToolMarkerVisibility::GetBrush));
+}
+
+FReply SNavigationToolMarkerVisibility::OnDragDetected(const FGeometry& MyGeometry, const FPointerEvent& MouseEvent)
+{
+	if (MouseEvent.IsMouseButtonDown(EKeys::LeftMouseButton))
+	{
+		return FReply::Handled().BeginDragDrop(FMarkerVisibilityDragDropOp::New(GetMarkerVisibility(), UndoTransaction));
+	}
+	return FReply::Unhandled();
+}
+
+void SNavigationToolMarkerVisibility::OnDragEnter(const FGeometry& MyGeometry, const FDragDropEvent& DragDropEvent)
+{
+	const TSharedPtr<FMarkerVisibilityDragDropOp> DragDropOp = DragDropEvent.GetOperationAs<FMarkerVisibilityDragDropOp>();
+	if (DragDropOp.IsValid())
+	{
+		const bool bNewVisible = (DragDropOp->MarkerVisibility != EItemMarkerVisibility::Visible);
+		SetMarkersVisible(bNewVisible);
+	}
+}
+
+FReply SNavigationToolMarkerVisibility::HandleClick()
+{
+	if (!IsVisibilityWidgetEnabled())
+	{
+		return FReply::Unhandled();
+	}
+
+	const TSharedPtr<INavigationToolView> ToolView = WeakView.Pin();
+	const FNavigationToolItemPtr Item = WeakItem.Pin();
+	const TSharedPtr<FNavigationToolMarkerVisibilityColumn> Column = WeakColumn.Pin();
+
+	if (!ToolView.IsValid() || !Item.IsValid() || !Column.IsValid())
+	{
+		return FReply::Unhandled();
+	}
+
+	// Open an undo transaction
+	UndoTransaction = MakeUnique<FScopedTransaction>(LOCTEXT("SetNavigationToolMarkerVisibility", "Set Marker Visibility"));
+
+	const bool bNewVisible = (GetMarkerVisibility() != EItemMarkerVisibility::Visible);
+
+	// We operate on all the selected items if the specified item is selected
+	if (ToolView->IsItemSelected(Item.ToSharedRef()))
+	{
+		for (const FNavigationToolItemPtr& SelectedItem : ToolView->GetSelectedItems())
+		{
+			if (IMarkerVisibilityExtension* const MarkerVisibilityItem = SelectedItem->CastTo<IMarkerVisibilityExtension>())
+			{
+				MarkerVisibilityItem->SetMarkerVisibility(bNewVisible);
+			}
+		}
+	}
+	else
+	{
+		SetMarkersVisible(bNewVisible);
+	}
+
+	return FReply::Handled().DetectDrag(SharedThis(this), EKeys::LeftMouseButton);
+}
+
+FReply SNavigationToolMarkerVisibility::OnMouseButtonDoubleClick(const FGeometry& InMyGeometry, const FPointerEvent& InMouseEvent)
+{
+	return HandleClick();
+}
+
+FReply SNavigationToolMarkerVisibility::OnMouseButtonDown(const FGeometry& MyGeometry, const FPointerEvent& MouseEvent)
+{
+	if (MouseEvent.GetEffectingButton() != EKeys::LeftMouseButton)
+	{
+		return FReply::Unhandled();
+	}
+
+	return HandleClick();
+}
+
+FReply SNavigationToolMarkerVisibility::OnMouseButtonUp(const FGeometry& MyGeometry, const FPointerEvent& MouseEvent)
+{
+	if (MouseEvent.GetEffectingButton() == EKeys::LeftMouseButton)
+	{
+		UndoTransaction.Reset();
+		return FReply::Handled();
+	}
+
+	return FReply::Unhandled();
+}
+
+void SNavigationToolMarkerVisibility::OnMouseCaptureLost(const FCaptureLostEvent& CaptureLostEvent)
+{
+	UndoTransaction.Reset();
+}
+
+const FSlateBrush* SNavigationToolMarkerVisibility::GetBrush() const
+{
+	return FAppStyle::GetBrush(TEXT("AnimTimeline.SectionMarker"));
+}
+
+FSlateColor SNavigationToolMarkerVisibility::GetForegroundColor() const
+{
+	const IMarkerVisibilityExtension* const MarkerVisibilityExt = WeakItem.Pin()->CastTo<IMarkerVisibilityExtension>();
+	if (!MarkerVisibilityExt)
+	{
+		return FLinearColor::Transparent;
+	}
+
+	const bool bIsItemHovered = WeakRowWidget.IsValid() && WeakRowWidget.Pin()->IsHovered();
+
+	if (IsHovered() || bIsItemHovered)
+	{
+		switch (GetMarkerVisibility())
+		{
+		case EItemMarkerVisibility::None:
+			return FStyleColors::White25;
+
+		case EItemMarkerVisibility::PartiallyVisible:
+			return FStyleColors::ForegroundHover;
+
+		case EItemMarkerVisibility::Visible:
+			return FStyleColors::ForegroundHover;
+		}
+	}
+	else
+	{
+		switch (GetMarkerVisibility())
+		{
+		case EItemMarkerVisibility::None:
+			return FStyleColors::Transparent;
+
+		case EItemMarkerVisibility::PartiallyVisible:
+			return FStyleColors::White25;
+
+		case EItemMarkerVisibility::Visible:
+			return FStyleColors::Foreground;
+		}
+	}
+
+	return FStyleColors::Transparent;
+}
+
+EItemMarkerVisibility SNavigationToolMarkerVisibility::GetMarkerVisibility() const
+{
+	return GetItemMarkerVisibility(WeakItem.Pin());
+}
+
+void SNavigationToolMarkerVisibility::SetMarkersVisible(const bool bInVisible)
+{
+	SetItemMarkerVisibility(WeakItem.Pin(), bInVisible);
+}
+
+}
+
+#undef LOCTEXT_NAMESPACE
diff --git a/Engine/Plugins/Experimental/SequenceNavigator/Source/SequenceNavigator/Private/Widgets/Columns/SNavigationToolMarkerVisibility.h b/Engine/Plugins/Experimental/SequenceNavigator/Source/SequenceNavigator/Private/Widgets/Columns/SNavigationToolMarkerVisibility.h
new file mode 100644
--- /dev/null
+++ b/Engine/Plugins/Experimental/SequenceNavigator/Source/SequenceNavigator/Private/Widgets/Columns/SNavigationToolMarkerVisibility.h
@@ -0,0 +1,64 @@
+// Copyright Epic Games, Inc. All Rights Reserved.
+
+#pragma once
+
+#include "Extensions/IMarkerVisibilityExtension.h"
+#include "NavigationToolDefines.h"
+#include "ScopedTransaction.h"
+#include "Templates/UniquePtr.h"
+#include "Widgets/Images/SImage.h"
+
+namespace UE::Sequencer::NavigationTool
+{
+
+class FNavigationToolMarkerVisibilityColumn;
+class INavigationToolView;
+class SNavigationToolTreeRow;
+
+/** Widget responsible for managing the visibility for a single item */
+class SNavigationToolMarkerVisibility : public SImage
+{
+public:
+	SLATE_BEGIN_ARGS(SNavigationToolMarkerVisibility) {}
+	SLATE_END_ARGS()
+
+	void Construct(const FArguments& InArgs
+		, const TSharedRef<FNavigationToolMarkerVisibilityColumn>& InColumn
+		, const FNavigationToolItemRef& InItem
+		, const TSharedRef<INavigationToolView>& InView
+		, const TSharedRef<SNavigationToolTreeRow>& InRowWidget);
+
+private:
+	/** Returns whether the widget is enabled or not */
+	virtual bool IsVisibilityWidgetEnabled() const { return true; }
+
+	virtual const FSlateBrush* GetBrush() const;
+
+	virtual FReply OnDragDetected(const FGeometry& MyGeometry, const FPointerEvent& MouseEvent) override;
+	virtual void OnDragEnter(const FGeometry& MyGeometry, const FDragDropEvent& DragDropEvent) override;
+
+	FReply HandleClick();
+
+	virtual FReply OnMouseButtonDoubleClick(const FGeometry& InMyGeometry, const FPointerEvent& InMouseEvent) override;
+	virtual FReply OnMouseButtonDown(const FGeometry& MyGeometry, const FPointerEvent& MouseEvent) override;
+	virtual FReply OnMouseButtonUp(const FGeometry& MyGeometry, const FPointerEvent& MouseEvent) override;
+	virtual void OnMouseCaptureLost(const FCaptureLostEvent& CaptureLostEvent) override;
+
+	virtual FSlateColor GetForegroundColor() const override;
+	
+	EItemMarkerVisibility GetMarkerVisibility() const;
+
+	void SetMarkersVisible(const bool bInVisible);
+
+	TWeakPtr<FNavigationToolMarkerVisibilityColumn> WeakColumn;
+
+	FNavigationToolItemWeakPtr WeakItem;
+
+	TWeakPtr<INavigationToolView> WeakView;
+
+	TWeakPtr<SNavigationToolTreeRow> WeakRowWidget;
+
+	TUniquePtr<FScopedTransaction> UndoTransaction;
+};
+
+}
diff --git a/Engine/Plugins/Experimental/SequenceNavigator/Source/SequenceNavigator/Private/Widgets/Columns/SNavigationToolOutTime.cpp b/Engine/Plugins/Experimental/SequenceNavigator/Source/SequenceNavigator/Private/Widgets/Columns/SNavigationToolOutTime.cpp
new file mode 100644
--- /dev/null
+++ b/Engine/Plugins/Experimental/SequenceNavigator/Source/SequenceNavigator/Private/Widgets/Columns/SNavigationToolOutTime.cpp
@@ -0,0 +1,72 @@
+// Copyright Epic Games, Inc. All Rights Reserved.
+
+#include "SNavigationToolOutTime.h"
+#include "Items/NavigationToolItemUtils.h"
+#include "Items/NavigationToolSequence.h"
+#include "NavigationToolView.h"
+#include "ScopedTransaction.h"
+#include "Widgets/SNavigationToolTreeRow.h"
+
+#define LOCTEXT_NAMESPACE "SNavigationToolOutTime"
+
+namespace UE::Sequencer::NavigationTool
+{
+
+void SNavigationToolOutTime::Construct(const FArguments& InArgs
+	, const FNavigationToolItemRef& InItem
+	, const TSharedRef<INavigationToolView>& InView
+	, const TSharedRef<SNavigationToolTreeRow>& InRowWidget)
+{
+	WeakItem = InItem;
+	WeakView = InView;
+	WeakRowWidget = InRowWidget;
+
+	if (!InItem->IsA<IOutTimeExtension>())
+	{
+		return;
+	}
+
+	SNavigationToolTime::Construct(SNavigationToolTime::FArguments(), InItem, InView, InRowWidget);
+}
+
+IOutTimeExtension* SNavigationToolOutTime::GetOutTimeExtension() const
+{
+	return WeakItem.IsValid() ? WeakItem.Pin()->CastTo<IOutTimeExtension>() : nullptr;
+}
+
+FName SNavigationToolOutTime::GetStyleName() const
+{
+	return TEXT("SpinBox.OutTime");
+}
+
+double SNavigationToolOutTime::GetFrameTimeValue() const
+{
+	if (const IOutTimeExtension* const OutTimeExtension = GetOutTimeExtension())
+	{
+		return OutTimeExtension->GetOutTime().Value;
+	}
+	return 0;
+}
+
+void SNavigationToolOutTime::OnFrameTimeValueCommitted(const double InNewValue, const ETextCommit::Type InCommitType)
+{
+	IOutTimeExtension* const OutTimeExtension = GetOutTimeExtension();
+	if (!OutTimeExtension)
+	{
+		return;
+	}
+
+	const bool bShouldTransact = !UndoTransaction.IsValid() && (InCommitType == ETextCommit::OnEnter);
+	const FScopedTransaction Transaction(GetTransactionText(), bShouldTransact);
+
+	return OutTimeExtension->SetOutTime(static_cast<int32>(InNewValue));
+}
+
+FText SNavigationToolOutTime::GetTransactionText() const
+{
+	return LOCTEXT("SetOutTimeTransaction", "Set Out Time");
+}
+
+}
+
+#undef LOCTEXT_NAMESPACE
diff --git a/Engine/Plugins/Experimental/SequenceNavigator/Source/SequenceNavigator/Private/Widgets/Columns/SNavigationToolRevisionControl.cpp b/Engine/Plugins/Experimental/SequenceNavigator/Source/SequenceNavigator/Private/Widgets/Columns/SNavigationToolRevisionControl.cpp
new file mode 100644
--- /dev/null
+++ b/Engine/Plugins/Experimental/SequenceNavigator/Source/SequenceNavigator/Private/Widgets/Columns/SNavigationToolRevisionControl.cpp
@@ -0,0 +1,108 @@
+// Copyright Epic Games, Inc. All Rights Reserved.
+
+#include "SNavigationToolRevisionControl.h"
+#include "Columns/NavigationToolRevisionControlColumn.h"
+#include "Extensions/IRevisionControlExtension.h"
+#include "Items/INavigationToolItem.h"
+#include "NavigationToolView.h"
+#include "Styling/StyleColors.h"
+#include "Widgets/SNavigationToolTreeRow.h"
+
+#define LOCTEXT_NAMESPACE "SNavigationToolRevisionControl"
+
+namespace UE::Sequencer::NavigationTool
+{
+
+EItemRevisionControlState GetRevisionControlState(const FNavigationToolItemPtr& InItem)
+{
+	if (!InItem.IsValid())
+	{
+		return EItemRevisionControlState::None;
+	}
+
+	const IRevisionControlExtension* const RevisionControlItem = InItem->CastTo<IRevisionControlExtension>();
+	if (!RevisionControlItem)
+	{
+		return EItemRevisionControlState::None;
+	}
+
+	return RevisionControlItem->GetRevisionControlState();
+}
+
+const FSlateBrush* GetRevisionControlStatusIcon(const FNavigationToolItemPtr& InItem)
+{
+	if (!InItem.IsValid())
+	{
+		return nullptr;
+	}
+
+	const IRevisionControlExtension* const RevisionControlItem = InItem->CastTo<IRevisionControlExtension>();
+	if (!RevisionControlItem)
+	{
+		return nullptr;
+	}
+
+	return RevisionControlItem->GetRevisionControlStatusIcon();
+}
+
+FText GetRevisionControlStatusText(const FNavigationToolItemPtr& InItem)
+{
+	if (!InItem.IsValid())
+	{
+		return FText::GetEmpty();
+	}
+
+	const IRevisionControlExtension* const RevisionControlItem = InItem->CastTo<IRevisionControlExtension>();
+	if (!RevisionControlItem)
+	{
+		return FText::GetEmpty();
+	}
+
+	return RevisionControlItem->GetRevisionControlStatusText();
+}
+
+void SNavigationToolRevisionControl::Construct(const FArguments& InArgs
+	, const TSharedRef<FNavigationToolRevisionControlColumn>& InColumn
+	, const FNavigationToolItemRef& InItem
+	, const TSharedRef<INavigationToolView>& InView
+	, const TSharedRef<SNavigationToolTreeRow>& InRowWidget)
+{
+	WeakColumn = InColumn;
+	WeakItem = InItem;
+	WeakView = InView;
+	WeakRowWidget = InRowWidget;
+
+	SetToolTipText(TAttribute<FText>::CreateSP(this, &SNavigationToolRevisionControl::GetToolTipText));
+
+	SImage::Construct(SImage::FArguments()
+		.Image(this, &SNavigationToolRevisionControl::GetBrush)
+		.ColorAndOpacity(this, &SNavigationToolRevisionControl::GetForegroundColor));
+}
+
+FSlateColor SNavigationToolRevisionControl::GetForegroundColor() const
+{
+	switch (GetRevisionControlState(WeakItem.Pin()))
+	{
+	case EItemRevisionControlState::None:
+		return FStyleColors::Transparent;
+	case EItemRevisionControlState::PartiallySourceControlled:
+		return FStyleColors::White25;
+	case EItemRevisionControlState::SourceControlled:
+		return FStyleColors::Foreground;
+	}
+	return FStyleColors::Transparent;
+}
+
+const FSlateBrush* SNavigationToolRevisionControl::GetBrush() const
+{
+	return GetRevisionControlStatusIcon(WeakItem.Pin());
+}
+
+FText SNavigationToolRevisionControl::GetToolTipText() const
+{
+	return GetRevisionControlStatusText(WeakItem.Pin());
+}
+
+}
+
+#undef LOCTEXT_NAMESPACE
diff --git a/Engine/Plugins/Experimental/SequenceNavigator/Source/SequenceNavigator/Private/Widgets/Columns/SNavigationToolTake.cpp b/Engine/Plugins/Experimental/SequenceNavigator/Source/SequenceNavigator/Private/Widgets/Columns/SNavigationToolTake.cpp
new file mode 100644
--- /dev/null
+++ b/Engine/Plugins/Experimental/SequenceNavigator/Source/SequenceNavigator/Private/Widgets/Columns/SNavigationToolTake.cpp
@@ -0,0 +1,394 @@
+// Copyright Epic Games, Inc. All Rights Reserved.
+
+#include "SNavigationToolTake.h"
+#include "AssetRegistry/AssetData.h"
+#include "Framework/MultiBox/MultiBoxBuilder.h"
+#include "ISequencer.h"
+#include "Items/NavigationToolItemUtils.h"
+#include "MovieSceneSequence.h"
+#include "MovieSceneToolHelpers.h"
+#include "NavigationTool.h"
+#include "NavigationToolView.h"
+#include "Styling/StyleColors.h"
+#include "Widgets/Input/SComboBox.h"
+#include "Widgets/Layout/SBox.h"
+#include "Widgets/SNavigationToolTreeRow.h"
+
+using namespace UE::Sequencer::NavigationTool;
+
+#define LOCTEXT_NAMESPACE "SNavigationToolTake"
+
+namespace UE::Sequencer::NavigationTool
+{
+
+DECLARE_DELEGATE_RetVal_OneParam(FReply, FOnTakeEntrySelected, const TSharedPtr<SNavigationToolTake::FTakeItemInfo> /*InTakeInfo*/);
+
+class SNavigationToolTakeEntry : public SCompoundWidget
+{
+public:
+	SLATE_BEGIN_ARGS(SNavigationToolTakeEntry) {}
+		SLATE_EVENT(FOnTakeEntrySelected, OnEntrySelected)
+	SLATE_END_ARGS()
+
+	void Construct(const FArguments& InArgs, const TSharedPtr<SNavigationToolTake::FTakeItemInfo>& InTakeEntry)
+	{
+		TakeEntry  = InTakeEntry;
+		MenuButtonStyle = &FAppStyle::Get().GetWidgetStyle<FButtonStyle>("Menu.Button");
+
+		OnTakeEntrySelected = InArgs._OnEntrySelected;
+
+		ChildSlot
+		[
+			SNew(SBox)
+			.WidthOverride(120.f)
+			[
+				SNew(SBorder)
+				.BorderImage(this, &SNavigationToolTakeEntry::GetBorderImage)
+				.Padding(FMargin(12.f, 1.f, 12.f, 1.f))
+				[
+					SNew(STextBlock)
+					.Text(FText::FromString(TakeEntry->DisplayName))
+					.ColorAndOpacity(FLinearColor::White)
+				]
+
+			]
+		];
+	}
+
+	const FSlateBrush* GetBorderImage() const
+	{
+		if (IsHovered())
+		{
+			return &MenuButtonStyle->Hovered;
+		}
+		return &MenuButtonStyle->Normal;
+	}
+
+	virtual FReply OnMouseButtonUp(const FGeometry& InGeometry, const FPointerEvent& InMouseEvent) override
+	{
+		if (OnTakeEntrySelected.IsBound())
+		{
+			return OnTakeEntrySelected.Execute(TakeEntry);
+		}
+		return FReply::Handled();
+	}
+
+protected:
+	TSharedPtr<SNavigationToolTake::FTakeItemInfo> TakeEntry;
+	FOnTakeEntrySelected OnTakeEntrySelected;
+	const FButtonStyle* MenuButtonStyle = nullptr;
+};
+
+class FTakeDragDropOp : public FDragDropOperation
+{
+public:
+	DRAG_DROP_OPERATOR_TYPE(FTakeDragDropOp, FDragDropOperation)
+
+	virtual TSharedPtr<SWidget> GetDefaultDecorator() const override
+	{
+		return SNullWidget::NullWidget;
+	}
+
+	static TSharedRef<FTakeDragDropOp> New(const TSharedPtr<SNavigationToolTake::FTakeItemInfo>& InTakeInfo)
+	{
+		const TSharedRef<FTakeDragDropOp> Operation = MakeShared<FTakeDragDropOp>();
+		Operation->TakeInfo = InTakeInfo;
+		Operation->Construct();
+		return Operation;
+	}
+
+	TSharedPtr<SNavigationToolTake::FTakeItemInfo> TakeInfo;
+};
+
+void SNavigationToolTake::Construct(const FArguments& InArgs
+	, const FNavigationToolItemRef& InItem
+	, const TSharedRef<INavigationToolView>& InView
+	, const TSharedRef<SNavigationToolTreeRow>& InRowWidget)
+{
+	WeakItem = InItem;
+	WeakView = InView;
+	WeakRowWidget = InRowWidget;
+
+	WeakTool = InView->GetOwnerTool();
+
+	CacheTakes();
+
+	if (!ActiveTakeInfo.IsValid())
+	{
+		return;
+	}
+
+	ChildSlot
+	[
+		SNew(SComboBox<TSharedPtr<SNavigationToolTake::FTakeItemInfo>>)
+		.OptionsSource(&CachedTakes)
+		.InitiallySelectedItem(ActiveTakeInfo)
+		.OnSelectionChanged(this, &SNavigationToolTake::OnSelectionChanged)
+		.OnGenerateWidget(this, &SNavigationToolTake::GenerateTakeWidget)
+		[
+			GenerateTakeWidget(ActiveTakeInfo)
+		]
+	];
+}
+
+TSharedRef<SWidget> SNavigationToolTake::GenerateTakeWidget(const TSharedPtr<FTakeItemInfo> InTakeInfo)
+{
+	const FText TakeInfoText = FText::Format(LOCTEXT("TakeNumberLabel", "({0}/{1})")
+		, FText::FromString(FString::FromInt(InTakeInfo->TakeIndex + 1))
+		, FText::FromString(FString::FromInt(CachedTakes.Num())));
+
+	return SNew(SHorizontalBox)
+		+ SHorizontalBox::Slot()
+		.HAlign(HAlign_Left)
+		.VAlign(VAlign_Center)
+		.AutoWidth()
+		[
+			SNew(STextBlock)
+			.Text(FText::FromString(FString::FromInt(InTakeInfo->TakeNumber)))
+		]
+		+ SHorizontalBox::Slot()
+		.FillWidth(1.f)
+		.HAlign(HAlign_Right)
+		.VAlign(VAlign_Center)
+		.Padding(5.f, 0.f, 0.f, 0.f)
+		[
+			SNew(STextBlock)
+			.TextStyle(FAppStyle::Get(), TEXT("SmallText"))
+			.ColorAndOpacity(FStyleColors::Hover)
+			.Text(TakeInfoText)
+		];
+}
+
+void SNavigationToolTake::RemoveItemColor() const
+{
+	if (const TSharedPtr<INavigationTool> Tool = WeakTool.Pin())
+	{
+		Tool->RemoveItemColor(WeakItem.Pin());
+	}
+}
+
+FSlateColor SNavigationToolTake::GetBorderColor() const
+{
+	if (IsHovered())
+	{
+		return FSlateColor::UseForeground();
+	}
+	return FSlateColor::UseSubduedForeground();
+}
+
+FReply SNavigationToolTake::OnTakeEntrySelected(const TSharedPtr<FTakeItemInfo> InTakeInfo)
+{
+	if (const TSharedPtr<INavigationTool> Tool = WeakTool.Pin())
+	{
+		SetActiveTake(InTakeInfo->WeakSequence.Get());
+
+		return FReply::Handled();
+	}
+
+	return FReply::Unhandled();
+}
+
+void SNavigationToolTake::Press()
+{
+	bIsPressed = true;
+}
+
+void SNavigationToolTake::Release()
+{
+	bIsPressed = false;
+}
+
+FReply SNavigationToolTake::OnMouseButtonDown(const FGeometry& MyGeometry, const FPointerEvent& MouseEvent)
+{
+	if (MouseEvent.IsMouseButtonDown(EKeys::LeftMouseButton))
+	{
+		Press();
+
+		const TSharedRef<SNavigationToolTake> This = SharedThis(this);
+		return FReply::Handled().CaptureMouse(This).DetectDrag(This, EKeys::LeftMouseButton);
+	}
+
+	return SCompoundWidget::OnMouseButtonDown(MyGeometry, MouseEvent);
+}
+
+FReply SNavigationToolTake::OnMouseButtonUp(const FGeometry& MyGeometry, const FPointerEvent& MouseEvent)
+{
+	FReply Reply = FReply::Unhandled();
+
+	if (bIsPressed && IsHovered() && MouseEvent.GetEffectingButton() == EKeys::LeftMouseButton)
+	{
+		Release();
+		Reply = FReply::Handled();
+	}
+
+	if (Reply.GetMouseCaptor().IsValid() == false && HasMouseCapture())
+	{
+		Reply.ReleaseMouseCapture();
+	}
+
+	return Reply;
+}
+
+void SNavigationToolTake::OnFocusLost(const FFocusEvent& InFocusEvent)
+{
+	Release();
+
+	SCompoundWidget::OnFocusLost(InFocusEvent);
+}
+
+FReply SNavigationToolTake::OnDragDetected(const FGeometry& MyGeometry, const FPointerEvent& MouseEvent)
+{
+	if (MouseEvent.IsMouseButtonDown(EKeys::LeftMouseButton))
+	{
+		return FReply::Handled().BeginDragDrop(FTakeDragDropOp::New(ActiveTakeInfo));
+	}
+
+	return SCompoundWidget::OnDragDetected(MyGeometry, MouseEvent);
+}
+
+void SNavigationToolTake::OnDragEnter(const FGeometry& MyGeometry, const FDragDropEvent& DragDropEvent)
+{
+	if (const TSharedPtr<FTakeDragDropOp> DragOp = DragDropEvent.GetOperationAs<FTakeDragDropOp>())
+	{
+		OnTakeEntrySelected(DragOp->TakeInfo);
+	}
+
+	SCompoundWidget::OnDragEnter(MyGeometry, DragDropEvent);
+}
+
+void SNavigationToolTake::OnSelectionChanged(const TSharedPtr<FTakeItemInfo> InTakeInfo, const ESelectInfo::Type InSelectType)
+{
+	if (InTakeInfo.IsValid() && InTakeInfo->WeakSequence.IsValid())
+	{
+		SetActiveTake(InTakeInfo->WeakSequence.Get());
+	}
+}
+
+void SNavigationToolTake::CacheTakes()
+{
+	CachedTakes.Reset();
+
+	UMovieSceneSubSection* const SubSection = Cast<UMovieSceneSubSection>(GetSequenceItemSubSection(WeakItem.Pin()));
+	if (!IsValid(SubSection))
+	{
+		return;
+	}
+
+	TArray<FAssetData> AssetData;
+	uint32 CurrentTakeNumber = INDEX_NONE;
+	MovieSceneToolHelpers::GatherTakes(SubSection, AssetData, CurrentTakeNumber);
+
+	AssetData.Sort([SubSection](const FAssetData& InA, const FAssetData& InB)
+		{
+			uint32 TakeNumberA = INDEX_NONE;
+			uint32 TakeNumberB = INDEX_NONE;
+			if (MovieSceneToolHelpers::GetTakeNumber(SubSection, InA, TakeNumberA)
+				&& MovieSceneToolHelpers::GetTakeNumber(SubSection, InB, TakeNumberB))
+			{
+				return TakeNumberA < TakeNumberB;
+			}
+			return true;
+		});
+
+	uint32 TakeIndex = 0;
+	for (const FAssetData& ThisAssetData : AssetData)
+	{
+		uint32 TakeNumber = INDEX_NONE;
+		if (MovieSceneToolHelpers::GetTakeNumber(SubSection, ThisAssetData, TakeNumber))
+		{
+			UMovieSceneSequence* const Sequence = Cast<UMovieSceneSequence>(ThisAssetData.GetAsset());
+			if (IsValid(Sequence))
+			{
+				const TSharedRef<FTakeItemInfo> NewTakeInfo = MakeShared<FTakeItemInfo>();
+				NewTakeInfo->TakeIndex = TakeIndex;
+				NewTakeInfo->TakeNumber = TakeNumber;
+				NewTakeInfo->DisplayName = Sequence->GetDisplayName().ToString();
+				NewTakeInfo->WeakSequence = Sequence;
+				CachedTakes.Add(NewTakeInfo);
+
+				if (TakeNumber == CurrentTakeNumber)
+				{
+					ActiveTakeInfo = NewTakeInfo;
+				}
+
+				++TakeIndex;
+			}
+		}
+	}
+}
+
+void SNavigationToolTake::SetActiveTake(UMovieSceneSequence* const InSequence)
+{
+	const TSharedPtr<INavigationTool> Tool = WeakTool.Pin();
+	if (!Tool.IsValid())
+	{
+		return;
+	}
+
+	const TSharedPtr<ISequencer> Sequencer = Tool->GetSequencer();
+	if (!Sequencer.IsValid())
+	{
+		return;
+	}
+
+	bool bChangedTake = false;
+
+	const FScopedTransaction Transaction(LOCTEXT("ChangeTake_Transaction", "Change Take"));
+
+	TArray<UMovieSceneSection*> Sections;
+	Sequencer->GetSelectedSections(Sections);
+
+	for (int32 SectionIndex = 0; SectionIndex < Sections.Num(); ++SectionIndex)
+	{
+		if (!Sections[SectionIndex]->IsA<UMovieSceneSubSection>())
+		{
+			continue;
+		}
+
+		UMovieSceneSubSection* Section = Cast<UMovieSceneSubSection>(Sections[SectionIndex]);
+		UMovieSceneSubTrack* SubTrack = CastChecked<UMovieSceneSubTrack>(Section->GetOuter());
+
+		TRange<FFrameNumber> NewSectionRange = Section->GetRange();
+		FFrameNumber		 NewSectionStartOffset = Section->Parameters.StartFrameOffset;
+		int32                NewSectionPrerollFrames = Section->GetPreRollFrames();
+		int32                NewRowIndex = Section->GetRowIndex();
+		FFrameNumber         NewSectionStartTime = NewSectionRange.GetLowerBound().IsClosed() ? UE::MovieScene::DiscreteInclusiveLower(NewSectionRange) : 0;
+		int32                NewSectionRowIndex = Section->GetRowIndex();
+		FColor               NewSectionColorTint = Section->GetColorTint();
+
+		const int32 Duration = (NewSectionRange.GetLowerBound().IsClosed() && NewSectionRange.GetUpperBound().IsClosed()) ? UE::MovieScene::DiscreteSize(NewSectionRange) : 1;
+		UMovieSceneSubSection* NewSection = SubTrack->AddSequence(InSequence, NewSectionStartTime, Duration);
+
+		if (NewSection != nullptr)
+		{
+			SubTrack->RemoveSection(*Section);
+
+			NewSection->SetRange(NewSectionRange);
+			NewSection->Parameters.StartFrameOffset = NewSectionStartOffset;
+			NewSection->Parameters.TimeScale = Section->Parameters.TimeScale.DeepCopy(NewSection);
+			NewSection->SetPreRollFrames(NewSectionPrerollFrames);
+			NewSection->SetRowIndex(NewSectionRowIndex);
+			NewSection->SetColorTint(NewSectionColorTint);
+
+			UMovieSceneCinematicShotSection* ShotSection = Cast<UMovieSceneCinematicShotSection>(Section);
+			UMovieSceneCinematicShotSection* NewShotSection = Cast<UMovieSceneCinematicShotSection>(NewSection);
+
+			// If the old shot's name is not the same as the sequence's name, assume the user had customized the shot name, so carry it over
+			if (ShotSection && NewShotSection && ShotSection->GetSequence() && ShotSection->GetShotDisplayName() != ShotSection->GetSequence()->GetName())
+			{
+				NewShotSection->SetShotDisplayName(ShotSection->GetShotDisplayName());
+			}
+
+			bChangedTake = true;
+		}
+	}
+
+	if (bChangedTake)
+	{
+		Sequencer->NotifyMovieSceneDataChanged(EMovieSceneDataChangeType::MovieSceneStructureItemsChanged);
+	}
+}
+
+}
+
+#undef LOCTEXT_NAMESPACE
diff --git a/Engine/Plugins/Experimental/SequenceNavigator/Source/SequenceNavigator/Private/Widgets/Columns/SNavigationToolTake.h b/Engine/Plugins/Experimental/SequenceNavigator/Source/SequenceNavigator/Private/Widgets/Columns/SNavigationToolTake.h
new file mode 100644
--- /dev/null
+++ b/Engine/Plugins/Experimental/SequenceNavigator/Source/SequenceNavigator/Private/Widgets/Columns/SNavigationToolTake.h
@@ -0,0 +1,80 @@
+// Copyright Epic Games, Inc. All Rights Reserved.
+
+#pragma once
+
+#include "Columns/NavigationToolTakeColumn.h"
+#include "NavigationToolDefines.h"
+#include "ScopedTransaction.h"
+#include "Templates/UniquePtr.h"
+#include "Widgets/SCompoundWidget.h"
+
+class UMovieSceneSequence;
+
+namespace UE::Sequencer::NavigationTool
+{
+
+class INavigationTool;
+class INavigationToolView;
+class SNavigationToolTreeRow;
+
+class SNavigationToolTake : public SCompoundWidget
+{
+public:
+	struct FTakeItemInfo
+	{
+		uint32 TakeIndex = 0;
+		uint32 TakeNumber = 0;
+		FString DisplayName;
+		TWeakObjectPtr<UMovieSceneSequence> WeakSequence;
+	};
+
+	SLATE_BEGIN_ARGS(SNavigationToolTake) {}
+	SLATE_END_ARGS()
+
+	void Construct(const FArguments& InArgs
+		, const FNavigationToolItemRef& InItem
+		, const TSharedRef<INavigationToolView>& InView
+		, const TSharedRef<SNavigationToolTreeRow>& InRowWidget);
+
+	void RemoveItemColor() const;
+
+	FSlateColor GetBorderColor() const;
+
+	FReply OnTakeEntrySelected(const TSharedPtr<FTakeItemInfo> InTakeInfo);
+
+	void Press();
+	void Release();
+
+	virtual FReply OnMouseButtonDown(const FGeometry& MyGeometry, const FPointerEvent& MouseEvent) override;
+	virtual FReply OnMouseButtonUp(const FGeometry& MyGeometry, const FPointerEvent& MouseEvent) override;
+	virtual void OnFocusLost(const FFocusEvent& InFocusEvent) override;
+
+	virtual FReply OnDragDetected(const FGeometry& MyGeometry, const FPointerEvent& MouseEvent) override;
+	virtual void OnDragEnter(const FGeometry& MyGeometry, const FDragDropEvent& DragDropEvent) override;
+
+private:
+	TSharedRef<SWidget> GenerateTakeWidget(const TSharedPtr<FTakeItemInfo> InTakeInfo);
+
+	void OnSelectionChanged(const TSharedPtr<FTakeItemInfo> InTakeInfo, const ESelectInfo::Type InSelectType);
+
+	void CacheTakes();
+	void SetActiveTake(UMovieSceneSequence* const InSequence);
+
+	FNavigationToolItemWeakPtr WeakItem;
+
+	TWeakPtr<INavigationToolView> WeakView;
+
+	TWeakPtr<SNavigationToolTreeRow> WeakRowWidget;
+
+	TWeakPtr<INavigationTool> WeakTool;
+
+	TArray<TSharedPtr<FTakeItemInfo>> CachedTakes;
+
+	TSharedPtr<FTakeItemInfo> ActiveTakeInfo;
+
+	bool bIsPressed = false;
+
+	TUniquePtr<FScopedTransaction> UndoTransaction;
+};
+
+}
diff --git a/Engine/Plugins/Experimental/SequenceNavigator/Source/SequenceNavigator/Private/Widgets/SNavigationToolTreeRow.cpp b/Engine/Plugins/Experimental/SequenceNavigator/Source/SequenceNavigator/Private/Widgets/SNavigationToolTreeRow.cpp
new file mode 100644
--- /dev/null
+++ b/Engine/Plugins/Experimental/SequenceNavigator/Source/SequenceNavigator/Private/Widgets/SNavigationToolTreeRow.cpp
@@ -0,0 +1,117 @@
+// Copyright Epic Games, Inc. All Rights Reserved.
+
+#include "Widgets/SNavigationToolTreeRow.h"
+#include "Columns/INavigationToolColumn.h"
+#include "Items/INavigationToolItem.h"
+#include "NavigationToolStyle.h"
+#include "NavigationToolView.h"
+#include "Widgets/SNavigationToolTreeView.h"
+
+#define LOCTEXT_NAMESPACE "SNavigationToolTreeRow"
+
+namespace UE::Sequencer::NavigationTool
+{
+
+void SNavigationToolTreeRow::Construct(const FArguments& InArgs
+	, const TSharedRef<FNavigationToolView>& InToolView
+	, const TSharedPtr<SNavigationToolTreeView>& InTreeView
+	, const FNavigationToolItemPtr& InItem)
+{
+	WeakToolView = InToolView;
+	WeakTreeView = InTreeView;
+	Item = InItem;
+	HighlightText = InArgs._HighlightText;
+
+	SetColorAndOpacity(TAttribute<FLinearColor>::CreateSP(&*InToolView, &FNavigationToolView::GetItemBrushColor, Item));
+
+	SMultiColumnTableRow::Construct(FSuperRowType::FArguments()
+			.Style(&FNavigationToolStyle::Get().GetWidgetStyle<FTableRowStyle>(TEXT("TableViewRow")))
+			.OnCanAcceptDrop(InToolView, &FNavigationToolView::OnCanDrop)
+			.OnDragDetected(InToolView, &FNavigationToolView::OnDragDetected, Item)
+			.OnDragEnter(InToolView, &FNavigationToolView::OnDragEnter, Item)
+			.OnDragLeave(InToolView, &FNavigationToolView::OnDragLeave, Item)
+			.OnAcceptDrop(InToolView, &FNavigationToolView::OnDrop)
+			.OnDrop(this, &SNavigationToolTreeRow::OnDefaultDrop)
+		, InTreeView.ToSharedRef());
+}
+
+TSharedRef<SWidget> SNavigationToolTreeRow::GenerateWidgetForColumn(const FName& InColumnName)
+{
+	if (Item.IsValid())
+	{
+		const TSharedPtr<FNavigationToolView> ToolView = WeakToolView.Pin();
+		check(ToolView.IsValid());
+		
+		if (const TSharedPtr<INavigationToolColumn> Column = ToolView->GetColumns().FindRef(InColumnName))
+		{
+			return Column->ConstructRowWidget(Item.ToSharedRef(), ToolView.ToSharedRef(), SharedThis(this));
+		}
+	}
+	return SNullWidget::NullWidget;
+}
+
+FReply SNavigationToolTreeRow::OnMouseButtonUp(const FGeometry& MyGeometry, const FPointerEvent& MouseEvent)
+{
+	const TSharedPtr<FNavigationToolView> ToolView = WeakToolView.Pin();
+	if (!ToolView.IsValid())
+	{
+		return FReply::Unhandled();
+	}
+
+	//Select Item and the Tree of Children it contains
+	if (MouseEvent.GetEffectingButton() == EKeys::MiddleMouseButton)
+	{
+		ENavigationToolItemSelectionFlags Flags = ENavigationToolItemSelectionFlags::IncludeChildren
+			| ENavigationToolItemSelectionFlags::SignalSelectionChange
+			| ENavigationToolItemSelectionFlags::ScrollIntoView;
+
+		if (MouseEvent.IsControlDown())
+		{
+			Flags |= ENavigationToolItemSelectionFlags::AppendToCurrentSelection;
+		}
+
+		ToolView->SelectItems({Item}, Flags);
+
+		Item->OnSelect();
+
+		return FReply::Handled();
+	}
+	
+	return SMultiColumnTableRow::OnMouseButtonUp(MyGeometry, MouseEvent);
+}
+
+FReply SNavigationToolTreeRow::OnMouseButtonDoubleClick(const FGeometry& InGeometry, const FPointerEvent& InMouseEvent)
+{
+	// Select Item and the Tree of Children it contains
+	if (InMouseEvent.GetEffectingButton() == EKeys::LeftMouseButton)
+	{
+		if (const TSharedPtr<FNavigationToolView> ToolView = WeakToolView.Pin())
+		{
+			Item->OnDoubleClick();
+
+			return FReply::Handled();
+		}
+	}
+
+	return SMultiColumnTableRow<FNavigationToolItemPtr>::OnMouseButtonDoubleClick(InGeometry, InMouseEvent);
+}
+
+TSharedPtr<FNavigationToolView> SNavigationToolTreeRow::GetToolView() const
+{
+	return WeakToolView.Pin();
+}
+
+FReply SNavigationToolTreeRow::OnDefaultDrop(const FDragDropEvent& InDragDropEvent) const
+{
+	if (const TSharedPtr<FNavigationToolView> ToolView = WeakToolView.Pin())
+	{
+		ToolView->SetDragIntoTreeRoot(false);
+	}
+
+	// Always return handled as no action should take place if the Drop wasn't accepted
+	return FReply::Handled();
+}
+
+}
+
+#undef LOCTEXT_NAMESPACE
diff --git a/Engine/Plugins/Experimental/SequenceNavigator/Source/SequenceNavigator/Private/Widgets/SNavigationToolTreeView.cpp b/Engine/Plugins/Experimental/SequenceNavigator/Source/SequenceNavigator/Private/Widgets/SNavigationToolTreeView.cpp
new file mode 100644
--- /dev/null
+++ b/Engine/Plugins/Experimental/SequenceNavigator/Source/SequenceNavigator/Private/Widgets/SNavigationToolTreeView.cpp
@@ -0,0 +1,279 @@
+// Copyright Epic Games, Inc. All Rights Reserved.
+
+#include "SNavigationToolTreeView.h"
+#include "Framework/Commands/UICommandList.h"
+#include "Items/INavigationToolItem.h"
+#include "NavigationToolView.h"
+
+namespace UE::Sequencer::NavigationTool
+{
+
+void SNavigationToolTreeView::Construct(const FArguments& InArgs, const TSharedPtr<FNavigationToolView>& InToolView)
+{
+	WeakToolView = InToolView;
+
+	if (InToolView.IsValid())
+	{
+		InToolView->SaveColumnState();
+	}
+
+	STreeView::Construct(InArgs._TreeViewArgs);
+}
+
+int32 SNavigationToolTreeView::GetItemIndex(const FNavigationToolItemPtr& Item) const
+{
+	if (SListView<FNavigationToolItemPtr>::HasValidItemsSource())
+	{
+		return SListView<FNavigationToolItemPtr>::GetItems().Find(TListTypeTraits<FNavigationToolItemPtr>::NullableItemTypeConvertToItemType(Item));
+	}
+	return INDEX_NONE;
+}
+
+void SNavigationToolTreeView::FocusOnItem(const FNavigationToolItemPtr& InItem)
+{
+	SelectorItem = InItem;
+	RangeSelectionStart = InItem;
+}
+
+void SNavigationToolTreeView::ScrollItemIntoView(const FNavigationToolItemPtr& InItem)
+{
+	RequestScrollIntoView(InItem, 0);
+}
+
+void SNavigationToolTreeView::UpdateItemExpansions(const FNavigationToolItemPtr& InItem)
+{
+	if (InItem.IsValid())
+	{
+		const FSparseItemInfo* const SparseItemInfo = SparseItemInfos.Find(InItem);
+
+		bool bIsExpanded = false;
+		bool bHasExpandedChildren = false;
+
+		if (SparseItemInfo)
+		{
+			bIsExpanded = SparseItemInfo->bIsExpanded;
+			bHasExpandedChildren = SparseItemInfo->bHasExpandedChildren;
+		}
+
+		//Skip to avoid redundancy
+		if (bIsExpanded || bHasExpandedChildren)
+		{
+			return;
+		}
+
+		TArray<FNavigationToolItemPtr> ItemsToCheck = InItem->GetChildren();
+		while (ItemsToCheck.Num() > 0)
+		{
+			FNavigationToolItemPtr ItemToCheck = ItemsToCheck.Pop();
+
+			if (ItemToCheck.IsValid())
+			{
+				if (IsItemExpanded(ItemToCheck))
+				{
+					bHasExpandedChildren = true;
+					break;
+				}
+				ItemsToCheck.Append(ItemToCheck->GetChildren());
+			}
+		}
+
+		if (bIsExpanded || bHasExpandedChildren)
+		{
+			SparseItemInfos.Add(InItem, FSparseItemInfo(bIsExpanded, bHasExpandedChildren));
+		}
+	}
+}
+
+void SNavigationToolTreeView::Private_SetItemSelection(const FNavigationToolItemPtr InItem, bool bShouldBeSelected, bool bWasUserDirected)
+{
+	const TSharedPtr<FNavigationToolView> ToolView = WeakToolView.Pin();
+	if (ToolView.IsValid() && !ToolView->IsToolLocked() && ToolView->CanSelectItem(InItem))
+	{
+		STreeView::Private_SetItemSelection(InItem, bShouldBeSelected, bWasUserDirected);
+	}
+}
+
+void SNavigationToolTreeView::Private_ClearSelection()
+{
+	const TSharedPtr<FNavigationToolView> ToolView = WeakToolView.Pin();
+	if (ToolView.IsValid() && !ToolView->IsToolLocked())
+	{
+		STreeView::Private_ClearSelection();
+	}
+}
+
+void SNavigationToolTreeView::Private_SignalSelectionChanged(ESelectInfo::Type SelectInfo)
+{
+	STreeView::Private_SignalSelectionChanged(SelectInfo);
+
+	const TItemSet AddedItems = SelectedItems.Difference(PreviousSelectedItems);
+	for (const FNavigationToolItemPtr& AddedItem : AddedItems)
+	{
+		if (AddedItem.IsValid())
+		{
+			AddedItem->OnItemSelectionChanged(true);
+		}
+	}
+
+	const TItemSet RemovedItems = PreviousSelectedItems.Difference(SelectedItems);
+	for (const FNavigationToolItemPtr& RemovedItem : RemovedItems)
+	{
+		if (RemovedItem.IsValid())
+		{
+			RemovedItem->OnItemSelectionChanged(false);
+		}
+	}
+	PreviousSelectedItems = SelectedItems;
+}
+
+void SNavigationToolTreeView::Private_UpdateParentHighlights()
+{
+	this->Private_ClearHighlightedItems();
+	for (TItemSet::TConstIterator SelectedItemIt(this->SelectedItems); SelectedItemIt; ++SelectedItemIt)
+	{
+		const FNavigationToolItemPtr SelectedItem(*SelectedItemIt);
+
+		// Sometimes selection events can come through before the Linearized List is built, so the item may not exist yet.
+		const int32 ItemIndex = LinearizedItems.Find(SelectedItem);
+		if (ItemIndex == INDEX_NONE)
+		{
+			if (SelectedItem.IsValid())
+			{
+				FNavigationToolItemPtr ParentItem = SelectedItem->GetParent();
+				while (ParentItem.IsValid())
+				{
+					if (LinearizedItems.Contains(ParentItem))
+					{
+						this->SetItemHighlighted(ParentItem, true);
+					}
+					ParentItem = ParentItem->GetParent();
+				}
+			}
+			continue;
+		}
+
+		if (DenseItemInfos.IsValidIndex(ItemIndex))
+		{
+			const FItemInfo& ItemInfo = DenseItemInfos[ItemIndex];
+			
+			int32 ParentIndex = ItemInfo.ParentIndex;
+			
+			while (ParentIndex != INDEX_NONE)
+			{
+				const FNavigationToolItemPtr& ParentItem = this->LinearizedItems[ParentIndex];
+				this->Private_SetItemHighlighted(ParentItem, true);
+
+				const FItemInfo& ParentItemInfo = DenseItemInfos[ParentIndex];
+				
+				ParentIndex = ParentItemInfo.ParentIndex;
+			}
+		}
+	}
+}
+
+FCursorReply SNavigationToolTreeView::OnCursorQuery(const FGeometry& MyGeometry, const FPointerEvent& CursorEvent) const
+{
+	if (IsRightClickScrolling() && CursorEvent.IsMouseButtonDown(EKeys::MiddleMouseButton))
+	{
+		// We hide the native cursor as we'll be drawing the software EMouseCursor::GrabHandClosed cursor
+		return FCursorReply::Cursor(EMouseCursor::None);
+	}
+	return STreeView::OnCursorQuery(MyGeometry, CursorEvent);
+}
+
+FReply SNavigationToolTreeView::OnMouseMove(const FGeometry& MyGeometry, const FPointerEvent& MouseEvent)
+{
+	if (MouseEvent.IsMouseButtonDown(EKeys::MiddleMouseButton) && !MouseEvent.IsTouchEvent())
+	{
+		// We only care about deltas along the scroll axis
+		FTableViewDimensions CursorDeltaDimensions(Orientation, MouseEvent.GetCursorDelta());
+		CursorDeltaDimensions.LineAxis = 0.f;
+
+		const float ScrollByAmount = CursorDeltaDimensions.ScrollAxis / MyGeometry.Scale;
+
+		// If scrolling with the right mouse button, we need to remember how much we scrolled.
+		// If we did not scroll at all, we will bring up the context menu when the mouse is released.
+		AmountScrolledWhileRightMouseDown += FMath::Abs(ScrollByAmount);
+
+		// Has the mouse moved far enough with the right mouse button held down to start capturing
+		// the mouse and dragging the view?
+		if (IsRightClickScrolling())
+		{
+			// Make sure the active timer is registered to update the inertial scroll
+			if (!bIsScrollingActiveTimerRegistered)
+			{
+				bIsScrollingActiveTimerRegistered = true;
+				RegisterActiveTimer(0.f, FWidgetActiveTimerDelegate::CreateSP(this, &SNavigationToolTreeView::UpdateInertialScroll));
+			}
+
+			TickScrollDelta -= ScrollByAmount;
+
+			const float AmountScrolled = this->ScrollBy(MyGeometry, -ScrollByAmount, AllowOverscroll);
+
+			FReply Reply = FReply::Handled();
+
+			// The mouse moved enough that we're now dragging the view. Capture the mouse
+			// so the user does not have to stay within the bounds of the list while dragging.
+			if (this->HasMouseCapture() == false)
+			{
+				Reply.CaptureMouse(AsShared()).UseHighPrecisionMouseMovement(AsShared());
+				SoftwareCursorPosition = MyGeometry.AbsoluteToLocal(MouseEvent.GetScreenSpacePosition());
+				bShowSoftwareCursor    = true;
+			}
+
+			// Check if the mouse has moved.
+			if (AmountScrolled != 0)
+			{
+				SoftwareCursorPosition += CursorDeltaDimensions.ToVector2D();
+			}
+
+			return Reply;
+		}
+	}
+
+	return STreeView::OnMouseMove(MyGeometry, MouseEvent);
+}
+
+FReply SNavigationToolTreeView::OnMouseButtonUp(const FGeometry& MyGeometry, const FPointerEvent& MouseEvent)
+{
+	if (MouseEvent.GetEffectingButton() == EKeys::MiddleMouseButton)
+	{
+		FReply Reply = FReply::Handled().ReleaseMouseCapture();
+		AmountScrolledWhileRightMouseDown = 0;
+		bShowSoftwareCursor = false;
+
+		// If we have mouse capture, snap the mouse back to the closest location that is within the list's bounds
+		if (HasMouseCapture())
+		{
+			const FSlateRect ListScreenSpaceRect = MyGeometry.GetLayoutBoundingRect();
+			const FVector2D CursorPosition = MyGeometry.LocalToAbsolute(SoftwareCursorPosition);
+
+			const FIntPoint BestPositionInList(
+					FMath::RoundToInt(FMath::Clamp(CursorPosition.X, ListScreenSpaceRect.Left, ListScreenSpaceRect.Right)),
+					FMath::RoundToInt(FMath::Clamp(CursorPosition.Y, ListScreenSpaceRect.Top, ListScreenSpaceRect.Bottom))
+				);
+
+			Reply.SetMousePos(BestPositionInList);
+		}
+
+		return Reply;
+	}
+	return STreeView::OnMouseButtonUp(MyGeometry, MouseEvent);
+}
+
+FReply SNavigationToolTreeView::OnKeyDown(const FGeometry& MyGeometry, const FKeyEvent& InKeyEvent)
+{
+	if (const TSharedPtr<FNavigationToolView> ToolView = WeakToolView.Pin())
+	{
+		ToolView->UpdateRecentViews();
+
+		TSharedPtr<FUICommandList> CommandList = ToolView->GetViewCommandList();
+		if (CommandList.IsValid() && CommandList->ProcessCommandBindings(InKeyEvent))
+		{
+			return FReply::Handled();	
+		}
+	}
+	return STreeView::OnKeyDown(MyGeometry, InKeyEvent);
+}
+
+}
diff --git a/Engine/Plugins/Experimental/SequenceNavigator/Source/SequenceNavigator/Private/Widgets/SNavigationToolTreeView.h b/Engine/Plugins/Experimental/SequenceNavigator/Source/SequenceNavigator/Private/Widgets/SNavigationToolTreeView.h
new file mode 100644
--- /dev/null
+++ b/Engine/Plugins/Experimental/SequenceNavigator/Source/SequenceNavigator/Private/Widgets/SNavigationToolTreeView.h
@@ -0,0 +1,53 @@
+// Copyright Epic Games, Inc. All Rights Reserved.
+
+#pragma once
+
+#include "NavigationToolDefines.h"
+#include "Widgets/Views/STreeView.h"
+
+namespace UE::Sequencer::NavigationTool
+{
+
+class FNavigationToolView;
+
+class SNavigationToolTreeView : public STreeView<FNavigationToolItemPtr>
+{
+public:
+	SLATE_BEGIN_ARGS(SNavigationToolTreeView) {}
+		SLATE_ARGUMENT(STreeView<FNavigationToolItemPtr>::FArguments, TreeViewArgs)
+	SLATE_END_ARGS()
+
+	void Construct(const FArguments& InArgs, const TSharedPtr<FNavigationToolView>& InToolView);
+
+	int32 GetItemIndex(const FNavigationToolItemPtr& Item) const;
+
+	void FocusOnItem(const FNavigationToolItemPtr& InItem);
+
+	void ScrollItemIntoView(const FNavigationToolItemPtr& InItem);
+
+	void UpdateItemExpansions(const FNavigationToolItemPtr& InItem);
+
+	//~ Begin STreeView
+	virtual void Private_UpdateParentHighlights() override;
+	//~ End STreeView
+
+	//~ Begin ITypedTableView
+	virtual void Private_SetItemSelection(const FNavigationToolItemPtr InItem, bool bShouldBeSelected, bool bWasUserDirected = false) override;
+	virtual void Private_SignalSelectionChanged(ESelectInfo::Type SelectInfo) override;
+	virtual void Private_ClearSelection() override;
+	//~ End SListView
+
+	//~ Begin SWidget
+	virtual FCursorReply OnCursorQuery(const FGeometry& MyGeometry, const FPointerEvent& CursorEvent) const override;
+	virtual FReply OnMouseMove(const FGeometry& MyGeometry, const FPointerEvent& MouseEvent) override;
+	virtual FReply OnMouseButtonUp(const FGeometry& MyGeometry, const FPointerEvent& MouseEvent) override;
+	virtual FReply OnKeyDown(const FGeometry& MyGeometry, const FKeyEvent& InKeyEvent) override;
+	//~ End SWidget
+
+private:
+	TWeakPtr<FNavigationToolView> WeakToolView;
+	
+	SNavigationToolTreeView::TItemSet PreviousSelectedItems;
+};
+
+}
diff --git a/Engine/Plugins/Experimental/SequenceNavigator/Source/SequenceNavigator/Private/Widgets/SNavigationToolView.h b/Engine/Plugins/Experimental/SequenceNavigator/Source/SequenceNavigator/Private/Widgets/SNavigationToolView.h
new file mode 100644
--- /dev/null
+++ b/Engine/Plugins/Experimental/SequenceNavigator/Source/SequenceNavigator/Private/Widgets/SNavigationToolView.h
@@ -0,0 +1,109 @@
+﻿// Copyright Epic Games, Inc. All Rights Reserved.
+
+#pragma once
+
+#include "Filters/SBasicFilterBar.h"
+#include "Items/INavigationToolItem.h"
+#include "NavigationToolDefines.h"
+#include "Widgets/SCompoundWidget.h"
+
+class ITableRow;
+class SBorder;
+class SHeaderRow;
+class STableViewBase;
+
+namespace UE::Sequencer::NavigationTool
+{
+
+class FNavigationToolFilter;
+class FNavigationToolFilterBar;
+class FNavigationToolItemDragDropOp;
+class FNavigationToolToolbarMenu;
+class FNavigationToolView;
+class SNavigationToolFilterBar;
+class SNavigationToolTreeView;
+enum class ENavigationToolFilterChange : uint8;
+
+class SNavigationToolView : public SCompoundWidget
+{
+public:
+	SLATE_BEGIN_ARGS(SNavigationToolView) {}
+	SLATE_END_ARGS()
+
+	void Construct(const FArguments& InArgs, const TSharedRef<FNavigationToolView>& InToolView);
+
+	virtual ~SNavigationToolView() override;
+
+	TSharedPtr<SNavigationToolTreeView> GetTreeView() const { return TreeView; }
+
+	/** Rebuilds the entire widget. Called when the filter bar layout is changed. */
+	void RebuildWidget();
+
+	void ReconstructColumns();
+
+	bool IsColumnVisible(const FName InColumnId) const;
+	void ShowHideColumn(const FName InColumnId, const bool bInVisible);
+
+	void SetColumnWidth(const FName InColumnId, const float InWidth);
+
+	void RequestTreeRefresh();
+
+	void SetItemSelection(const TArray<FNavigationToolItemPtr>& InItems, bool bSignalSelectionChange);
+
+	void OnItemSelectionChanged(const FNavigationToolItemPtr InItem, ESelectInfo::Type InSelectionType);
+
+	void ScrollItemIntoView(const FNavigationToolItemPtr& InItem) const;
+
+	bool IsItemExpanded(const FNavigationToolItemPtr& InItem) const;
+	void SetItemExpansion(const FNavigationToolItemPtr& Item, const bool bInExpand) const;
+
+	void UpdateItemExpansions(const FNavigationToolItemPtr& Item) const;
+
+	TSharedRef<ITableRow> OnItemGenerateRow(const FNavigationToolItemPtr Item, const TSharedRef<STableViewBase>& OwnerTable);
+
+	void SetKeyboardFocus() const;
+
+	virtual void OnMouseEnter(const FGeometry& MyGeometry, const FPointerEvent& MouseEvent) override;
+	virtual FReply OnKeyDown(const FGeometry& MyGeometry, const FKeyEvent& InKeyEvent) override;
+	virtual FReply OnDrop(const FGeometry& MyGeometry, const FDragDropEvent& DragDropEvent) override;
+	virtual void OnDragEnter(const FGeometry& MyGeometry, const FDragDropEvent& DragDropEvent) override;
+	virtual void OnDragLeave(const FDragDropEvent& DragDropEvent) override;
+
+	void SetTreeBorderVisibility(bool bVisible);
+
+	void GenerateColumnState(const FName InColumnId, FNavigationToolViewColumnSaveState& OutColumnState);
+	void GenerateColumnStateMap(TMap<FName, FNavigationToolViewColumnSaveState>& OutStateMap);
+
+	EFilterBarLayout GetFilterBarLayout() const;
+	void SetFilterBarLayout(const EFilterBarLayout InLayout);
+
+private:
+	TSharedRef<SWidget> ConstructSplitterContent();
+	TSharedRef<SWidget> ConstructMainContent();
+
+	void RebuildSearchAndFilterRow();
+
+	void OnFilterBarStateChanged(const bool InIsVisible, const EFilterBarLayout InNewLayout);
+	void OnFiltersChanged(const ENavigationToolFilterChange InChangeType, const TSharedRef<FNavigationToolFilter>& InFilter);
+
+	TSharedPtr<FNavigationToolFilterBar> GetFilterBar() const;
+
+	TWeakPtr<FNavigationToolView> WeakToolView;
+
+	TSharedPtr<FNavigationToolToolbarMenu> ToolbarMenu;
+
+	TSharedPtr<SVerticalBox> SearchAndFilterRow;
+	TSharedPtr<SHeaderRow> HeaderRowWidget;
+	TSharedPtr<SBorder> TreeBorder;
+	TSharedPtr<SNavigationToolTreeView> TreeView;
+	TSharedPtr<SNavigationToolFilterBar> FilterBarWidget;
+
+	/** If true the SuggestionList shouldn't appear since we selected what we wanted */
+	bool bIsEnterLastKeyPressed = false;
+
+	bool bSelectingItems = false;
+
+	TSet<TWeakPtr<FNavigationToolItemDragDropOp>> ItemDragDropOps;
+};
+
+}
diff --git a/Engine/Plugins/Experimental/SequenceNavigator/Source/SequenceNavigator/Public/Columns/NavigationToolColorColumn.h b/Engine/Plugins/Experimental/SequenceNavigator/Source/SequenceNavigator/Public/Columns/NavigationToolColorColumn.h
new file mode 100644
--- /dev/null
+++ b/Engine/Plugins/Experimental/SequenceNavigator/Source/SequenceNavigator/Public/Columns/NavigationToolColorColumn.h
@@ -0,0 +1,31 @@
+// Copyright Epic Games, Inc. All Rights Reserved.
+
+#pragma once
+
+#include "Columns/NavigationToolColumn.h"
+
+namespace UE::Sequencer::NavigationTool
+{
+
+class INavigationToolView;
+
+class FNavigationToolColorColumn : public FNavigationToolColumn
+{
+public:
+	UE_NAVIGATIONTOOL_INHERITS(FNavigationToolColorColumn, FNavigationToolColumn);
+
+	static FName StaticColumnId() { return TEXT("Color"); }
+
+protected:
+	//~ Begin INavigationToolColumn
+	virtual FName GetColumnId() const override { return StaticColumnId(); }
+	SEQUENCENAVIGATOR_API virtual FText GetColumnDisplayNameText() const override;
+	virtual bool ShouldShowColumnByDefault() const override { return true; }
+	SEQUENCENAVIGATOR_API virtual SHeaderRow::FColumn::FArguments ConstructHeaderRowColumn(const TSharedRef<INavigationToolView>& InView, const float InFillSize) override;
+	SEQUENCENAVIGATOR_API virtual TSharedRef<SWidget> ConstructRowWidget(const FNavigationToolItemRef& InItem
+		, const TSharedRef<INavigationToolView>& InView
+		, const TSharedRef<SNavigationToolTreeRow>& InRow) override;
+	//~ End INavigationToolColumn
+};
+
+}
diff --git a/Engine/Plugins/Experimental/SequenceNavigator/Source/SequenceNavigator/Public/Columns/NavigationToolColumn.h b/Engine/Plugins/Experimental/SequenceNavigator/Source/SequenceNavigator/Public/Columns/NavigationToolColumn.h
new file mode 100644
--- /dev/null
+++ b/Engine/Plugins/Experimental/SequenceNavigator/Source/SequenceNavigator/Public/Columns/NavigationToolColumn.h
@@ -0,0 +1,16 @@
+﻿// Copyright Epic Games, Inc. All Rights Reserved.
+
+#pragma once
+
+#include "Columns/INavigationToolColumn.h"
+
+namespace UE::Sequencer::NavigationTool
+{
+
+/** Base Navigation Tool Column to extend from */
+class FNavigationToolColumn : public INavigationToolColumn
+{
+
+};
+
+}
diff --git a/Engine/Plugins/Experimental/SequenceNavigator/Source/SequenceNavigator/Public/Columns/NavigationToolCommentColumn.h b/Engine/Plugins/Experimental/SequenceNavigator/Source/SequenceNavigator/Public/Columns/NavigationToolCommentColumn.h
new file mode 100644
--- /dev/null
+++ b/Engine/Plugins/Experimental/SequenceNavigator/Source/SequenceNavigator/Public/Columns/NavigationToolCommentColumn.h
@@ -0,0 +1,34 @@
+// Copyright Epic Games, Inc. All Rights Reserved.
+
+#pragma once
+
+#include "Columns/NavigationToolColumn.h"
+
+namespace UE::Sequencer::NavigationTool
+{
+
+class INavigationToolView;
+
+class FNavigationToolCommentColumn : public FNavigationToolColumn
+{
+public:
+	UE_NAVIGATIONTOOL_INHERITS(FNavigationToolCommentColumn, FNavigationToolColumn);
+
+	static FName StaticColumnId() { return TEXT("Comment"); }
+
+protected:
+	//~ Begin INavigationToolColumn
+	virtual FName GetColumnId() const override { return StaticColumnId(); }
+	SEQUENCENAVIGATOR_API virtual FText GetColumnDisplayNameText() const override;
+	SEQUENCENAVIGATOR_API virtual const FSlateBrush* GetIconBrush() const override;
+	virtual bool ShouldShowColumnByDefault() const override { return false; }
+	virtual float GetFillWidth() const override { return 5.f; }
+	SEQUENCENAVIGATOR_API virtual SHeaderRow::FColumn::FArguments ConstructHeaderRowColumn(const TSharedRef<INavigationToolView>& InView
+		, const float InFillSize) override;
+	SEQUENCENAVIGATOR_API virtual TSharedRef<SWidget> ConstructRowWidget(const FNavigationToolItemRef& InItem
+		, const TSharedRef<INavigationToolView>& InView
+		, const TSharedRef<SNavigationToolTreeRow>& InRow) override;
+	//~ End INavigationToolColumn
+};
+
+}
diff --git a/Engine/Plugins/Experimental/SequenceNavigator/Source/SequenceNavigator/Public/Columns/NavigationToolDeactiveStateColumn.h b/Engine/Plugins/Experimental/SequenceNavigator/Source/SequenceNavigator/Public/Columns/NavigationToolDeactiveStateColumn.h
new file mode 100644
--- /dev/null
+++ b/Engine/Plugins/Experimental/SequenceNavigator/Source/SequenceNavigator/Public/Columns/NavigationToolDeactiveStateColumn.h
@@ -0,0 +1,33 @@
+// Copyright Epic Games, Inc. All Rights Reserved.
+
+#pragma once
+
+#include "Columns/NavigationToolColumn.h"
+
+namespace UE::Sequencer::NavigationTool
+{
+
+class INavigationToolView;
+
+class FNavigationToolDeactiveStateColumn : public FNavigationToolColumn
+{
+public:
+	UE_NAVIGATIONTOOL_INHERITS(FNavigationToolDeactiveStateColumn, FNavigationToolColumn);
+
+	static FName StaticColumnId() { return TEXT("DeactiveState"); }
+
+protected:
+	//~ Begin INavigationToolColumn
+	virtual FName GetColumnId() const override { return StaticColumnId(); }
+	SEQUENCENAVIGATOR_API virtual FText GetColumnDisplayNameText() const override;
+	SEQUENCENAVIGATOR_API virtual const FSlateBrush* GetIconBrush() const override;
+	SEQUENCENAVIGATOR_API virtual SHeaderRow::FColumn::FArguments ConstructHeaderRowColumn(const TSharedRef<INavigationToolView>& InView
+		, const float InFillSize) override;
+	virtual bool ShouldShowColumnByDefault() const override { return true; }
+	SEQUENCENAVIGATOR_API virtual TSharedRef<SWidget> ConstructRowWidget(const FNavigationToolItemRef& InItem
+		, const TSharedRef<INavigationToolView>& InView
+		, const TSharedRef<SNavigationToolTreeRow>& InRow) override;
+	//~ End INavigationToolColumn
+};
+
+}
diff --git a/Engine/Plugins/Experimental/SequenceNavigator/Source/SequenceNavigator/Public/Columns/NavigationToolHBiasColumn.h b/Engine/Plugins/Experimental/SequenceNavigator/Source/SequenceNavigator/Public/Columns/NavigationToolHBiasColumn.h
new file mode 100644
--- /dev/null
+++ b/Engine/Plugins/Experimental/SequenceNavigator/Source/SequenceNavigator/Public/Columns/NavigationToolHBiasColumn.h
@@ -0,0 +1,34 @@
+// Copyright Epic Games, Inc. All Rights Reserved.
+
+#pragma once
+
+#include "Columns/NavigationToolColumn.h"
+
+namespace UE::Sequencer::NavigationTool
+{
+
+class INavigationToolView;
+
+class FNavigationToolHBiasColumn : public FNavigationToolColumn
+{
+public:
+	UE_NAVIGATIONTOOL_INHERITS(FNavigationToolHBiasColumn, FNavigationToolColumn);
+
+	static FName StaticColumnId() { return TEXT("HBias"); }
+
+protected:
+	//~ Begin INavigationToolColumn
+	virtual FName GetColumnId() const override { return StaticColumnId(); }
+	SEQUENCENAVIGATOR_API virtual FText GetColumnDisplayNameText() const override;
+	SEQUENCENAVIGATOR_API virtual const FSlateBrush* GetIconBrush() const override;
+	virtual bool ShouldShowColumnByDefault() const override { return false; }
+	virtual float GetFillWidth() const override { return 5.f; }
+	SEQUENCENAVIGATOR_API virtual SHeaderRow::FColumn::FArguments ConstructHeaderRowColumn(const TSharedRef<INavigationToolView>& InView
+		, const float InFillSize) override;
+	SEQUENCENAVIGATOR_API virtual TSharedRef<SWidget> ConstructRowWidget(const FNavigationToolItemRef& InItem
+		, const TSharedRef<INavigationToolView>& InView
+		, const TSharedRef<SNavigationToolTreeRow>& InRow) override;
+	//~ End INavigationToolColumn
+};
+
+}
diff --git a/Engine/Plugins/Experimental/SequenceNavigator/Source/SequenceNavigator/Public/Columns/NavigationToolInTimeColumn.h b/Engine/Plugins/Experimental/SequenceNavigator/Source/SequenceNavigator/Public/Columns/NavigationToolInTimeColumn.h
new file mode 100644
--- /dev/null
+++ b/Engine/Plugins/Experimental/SequenceNavigator/Source/SequenceNavigator/Public/Columns/NavigationToolInTimeColumn.h
@@ -0,0 +1,34 @@
+// Copyright Epic Games, Inc. All Rights Reserved.
+
+#pragma once
+
+#include "Columns/NavigationToolColumn.h"
+
+namespace UE::Sequencer::NavigationTool
+{
+
+class INavigationToolView;
+
+class FNavigationToolInTimeColumn : public FNavigationToolColumn
+{
+public:
+	UE_NAVIGATIONTOOL_INHERITS(FNavigationToolInTimeColumn, FNavigationToolColumn);
+
+	static FName StaticColumnId() { return TEXT("InTime"); }
+
+protected:
+	//~ Begin INavigationToolColumn
+	virtual FName GetColumnId() const override { return StaticColumnId(); }
+	SEQUENCENAVIGATOR_API virtual FText GetColumnDisplayNameText() const override;
+	SEQUENCENAVIGATOR_API virtual const FSlateBrush* GetIconBrush() const override;
+	virtual bool ShouldShowColumnByDefault() const override { return true; }
+	virtual float GetFillWidth() const override { return 5.f; }
+	SEQUENCENAVIGATOR_API virtual SHeaderRow::FColumn::FArguments ConstructHeaderRowColumn(const TSharedRef<INavigationToolView>& InView
+		, const float InFillSize) override;
+	SEQUENCENAVIGATOR_API virtual TSharedRef<SWidget> ConstructRowWidget(const FNavigationToolItemRef& InItem
+		, const TSharedRef<INavigationToolView>& InView
+		, const TSharedRef<SNavigationToolTreeRow>& InRow) override;
+	//~ End INavigationToolColumn
+};
+
+}
diff --git a/Engine/Plugins/Experimental/SequenceNavigator/Source/SequenceNavigator/Public/Columns/NavigationToolItemsColumn.h b/Engine/Plugins/Experimental/SequenceNavigator/Source/SequenceNavigator/Public/Columns/NavigationToolItemsColumn.h
new file mode 100644
--- /dev/null
+++ b/Engine/Plugins/Experimental/SequenceNavigator/Source/SequenceNavigator/Public/Columns/NavigationToolItemsColumn.h
@@ -0,0 +1,34 @@
+// Copyright Epic Games, Inc. All Rights Reserved.
+
+#pragma once
+
+#include "Columns/NavigationToolColumn.h"
+
+namespace UE::Sequencer::NavigationTool
+{
+
+class INavigationToolView;
+
+class FNavigationToolItemsColumn : public FNavigationToolColumn
+{
+public:
+	UE_NAVIGATIONTOOL_INHERITS(FNavigationToolItemsColumn, FNavigationToolColumn);
+
+	static FName StaticColumnId() { return TEXT("Items"); }
+
+protected:
+	//~ Begin INavigationToolColumn
+	virtual FName GetColumnId() const override { return StaticColumnId(); }
+	SEQUENCENAVIGATOR_API virtual FText GetColumnDisplayNameText() const override;
+	SEQUENCENAVIGATOR_API virtual const FSlateBrush* GetIconBrush() const override;
+	virtual bool ShouldShowColumnByDefault() const override { return true; }
+	virtual float GetFillWidth() const override { return 3.f; }
+	SEQUENCENAVIGATOR_API virtual SHeaderRow::FColumn::FArguments ConstructHeaderRowColumn(const TSharedRef<INavigationToolView>& InView
+		, const float InFillSize) override;
+	SEQUENCENAVIGATOR_API virtual TSharedRef<SWidget> ConstructRowWidget(const FNavigationToolItemRef& InItem
+		, const TSharedRef<INavigationToolView>& InView
+		, const TSharedRef<SNavigationToolTreeRow>& InRow) override;
+	//~ End INavigationToolColumn
+};
+
+}
diff --git a/Engine/Plugins/Experimental/SequenceNavigator/Source/SequenceNavigator/Public/Columns/NavigationToolLabelColumn.h b/Engine/Plugins/Experimental/SequenceNavigator/Source/SequenceNavigator/Public/Columns/NavigationToolLabelColumn.h
new file mode 100644
--- /dev/null
+++ b/Engine/Plugins/Experimental/SequenceNavigator/Source/SequenceNavigator/Public/Columns/NavigationToolLabelColumn.h
@@ -0,0 +1,34 @@
+// Copyright Epic Games, Inc. All Rights Reserved.
+
+#pragma once
+
+#include "Columns/NavigationToolColumn.h"
+
+namespace UE::Sequencer::NavigationTool
+{
+
+class INavigationToolView;
+
+class FNavigationToolLabelColumn : public FNavigationToolColumn
+{
+public:
+	UE_NAVIGATIONTOOL_INHERITS(FNavigationToolLabelColumn, FNavigationToolColumn);
+
+	static FName StaticColumnId() { return TEXT("Label"); }
+
+protected:
+	//~ Begin INavigationToolColumn
+	virtual FName GetColumnId() const override { return StaticColumnId(); }
+	SEQUENCENAVIGATOR_API virtual FText GetColumnDisplayNameText() const override;
+	SEQUENCENAVIGATOR_API virtual const FSlateBrush* GetIconBrush() const override;
+	SEQUENCENAVIGATOR_API virtual SHeaderRow::FColumn::FArguments ConstructHeaderRowColumn(const TSharedRef<INavigationToolView>& InView
+		, const float InFillSize) override;
+	virtual float GetFillWidth() const override { return 10.f; }
+	virtual bool CanHideColumn(const FName InColumnId) const override { return false; }
+	SEQUENCENAVIGATOR_API virtual TSharedRef<SWidget> ConstructRowWidget(const FNavigationToolItemRef& InItem
+		, const TSharedRef<INavigationToolView>& InView
+		, const TSharedRef<SNavigationToolTreeRow>& InRow) override;
+	//~ End INavigationToolColumn
+};
+
+}
diff --git a/Engine/Plugins/Experimental/SequenceNavigator/Source/SequenceNavigator/Public/Columns/NavigationToolLengthColumn.h b/Engine/Plugins/Experimental/SequenceNavigator/Source/SequenceNavigator/Public/Columns/NavigationToolLengthColumn.h
new file mode 100644
--- /dev/null
+++ b/Engine/Plugins/Experimental/SequenceNavigator/Source/SequenceNavigator/Public/Columns/NavigationToolLengthColumn.h
@@ -0,0 +1,34 @@
+// Copyright Epic Games, Inc. All Rights Reserved.
+
+#pragma once
+
+#include "Columns/NavigationToolColumn.h"
+
+namespace UE::Sequencer::NavigationTool
+{
+
+class INavigationToolView;
+
+class FNavigationToolLengthColumn : public FNavigationToolColumn
+{
+public:
+	UE_NAVIGATIONTOOL_INHERITS(FNavigationToolLengthColumn, FNavigationToolColumn);
+
+	static FName StaticColumnId() { return TEXT("Length"); }
+
+protected:
+	//~ Begin INavigationToolColumn
+	virtual FName GetColumnId() const override { return StaticColumnId(); }
+	SEQUENCENAVIGATOR_API virtual FText GetColumnDisplayNameText() const override;
+	SEQUENCENAVIGATOR_API virtual const FSlateBrush* GetIconBrush() const override;
+	virtual bool ShouldShowColumnByDefault() const override { return false; }
+	virtual float GetFillWidth() const override { return 5.f; }
+	SEQUENCENAVIGATOR_API virtual SHeaderRow::FColumn::FArguments ConstructHeaderRowColumn(const TSharedRef<INavigationToolView>& InView
+		, const float InFillSize) override;
+	SEQUENCENAVIGATOR_API virtual TSharedRef<SWidget> ConstructRowWidget(const FNavigationToolItemRef& InItem
+		, const TSharedRef<INavigationToolView>& InView
+		, const TSharedRef<SNavigationToolTreeRow>& InRow) override;
+	//~ End INavigationToolColumn
+};
+
+}
diff --git a/Engine/Plugins/Experimental/SequenceNavigator/Source/SequenceNavigator/Public/Columns/NavigationToolLockColumn.h b/Engine/Plugins/Experimental/SequenceNavigator/Source/SequenceNavigator/Public/Columns/NavigationToolLockColumn.h
new file mode 100644
--- /dev/null
+++ b/Engine/Plugins/Experimental/SequenceNavigator/Source/SequenceNavigator/Public/Columns/NavigationToolLockColumn.h
@@ -0,0 +1,33 @@
+// Copyright Epic Games, Inc. All Rights Reserved.
+
+#pragma once
+
+#include "Columns/NavigationToolColumn.h"
+
+namespace UE::Sequencer::NavigationTool
+{
+
+class INavigationToolView;
+
+class FNavigationToolLockColumn : public FNavigationToolColumn
+{
+public:
+	UE_NAVIGATIONTOOL_INHERITS(FNavigationToolLockColumn, FNavigationToolColumn);
+
+	static FName StaticColumnId() { return TEXT("Lock"); }
+
+protected:
+	//~ Begin INavigationToolColumn
+	virtual FName GetColumnId() const override { return StaticColumnId(); }
+	SEQUENCENAVIGATOR_API virtual FText GetColumnDisplayNameText() const override;
+	SEQUENCENAVIGATOR_API virtual const FSlateBrush* GetIconBrush() const override;
+	virtual bool ShouldShowColumnByDefault() const override { return true; }
+	SEQUENCENAVIGATOR_API virtual SHeaderRow::FColumn::FArguments ConstructHeaderRowColumn(const TSharedRef<INavigationToolView>& InView
+		, const float InFillSize) override;
+	SEQUENCENAVIGATOR_API virtual TSharedRef<SWidget> ConstructRowWidget(const FNavigationToolItemRef& InItem
+		, const TSharedRef<INavigationToolView>& InView
+		, const TSharedRef<SNavigationToolTreeRow>& InRow) override;
+	//~ End INavigationToolColumn
+};
+
+}
diff --git a/Engine/Plugins/Experimental/SequenceNavigator/Source/SequenceNavigator/Public/Columns/NavigationToolMarkerVisibilityColumn.h b/Engine/Plugins/Experimental/SequenceNavigator/Source/SequenceNavigator/Public/Columns/NavigationToolMarkerVisibilityColumn.h
new file mode 100644
--- /dev/null
+++ b/Engine/Plugins/Experimental/SequenceNavigator/Source/SequenceNavigator/Public/Columns/NavigationToolMarkerVisibilityColumn.h
@@ -0,0 +1,33 @@
+// Copyright Epic Games, Inc. All Rights Reserved.
+
+#pragma once
+
+#include "Columns/NavigationToolColumn.h"
+
+namespace UE::Sequencer::NavigationTool
+{
+
+class INavigationToolView;
+
+class FNavigationToolMarkerVisibilityColumn : public FNavigationToolColumn
+{
+public:
+	UE_NAVIGATIONTOOL_INHERITS(FNavigationToolMarkerVisibilityColumn, FNavigationToolColumn);
+
+	static FName StaticColumnId() { return TEXT("MarkerVisibility"); }
+
+protected:
+	//~ Begin INavigationToolColumn
+	virtual FName GetColumnId() const override { return StaticColumnId(); }
+	SEQUENCENAVIGATOR_API virtual FText GetColumnDisplayNameText() const override;
+	SEQUENCENAVIGATOR_API virtual const FSlateBrush* GetIconBrush() const override;
+	SEQUENCENAVIGATOR_API virtual SHeaderRow::FColumn::FArguments ConstructHeaderRowColumn(const TSharedRef<INavigationToolView>& InView
+		, const float InFillSize) override;
+	virtual bool ShouldShowColumnByDefault() const override { return true; }
+	SEQUENCENAVIGATOR_API virtual TSharedRef<SWidget> ConstructRowWidget(const FNavigationToolItemRef& InItem
+		, const TSharedRef<INavigationToolView>& InView
+		, const TSharedRef<SNavigationToolTreeRow>& InRow) override;
+	//~ End INavigationToolColumn
+};
+
+}
diff --git a/Engine/Plugins/Experimental/SequenceNavigator/Source/SequenceNavigator/Public/Columns/NavigationToolOutTimeColumn.h b/Engine/Plugins/Experimental/SequenceNavigator/Source/SequenceNavigator/Public/Columns/NavigationToolOutTimeColumn.h
new file mode 100644
--- /dev/null
+++ b/Engine/Plugins/Experimental/SequenceNavigator/Source/SequenceNavigator/Public/Columns/NavigationToolOutTimeColumn.h
@@ -0,0 +1,34 @@
+// Copyright Epic Games, Inc. All Rights Reserved.
+
+#pragma once
+
+#include "Columns/NavigationToolColumn.h"
+
+namespace UE::Sequencer::NavigationTool
+{
+
+class INavigationToolView;
+
+class FNavigationToolOutTimeColumn : public FNavigationToolColumn
+{
+public:
+	UE_NAVIGATIONTOOL_INHERITS(FNavigationToolOutTimeColumn, FNavigationToolColumn);
+
+	static FName StaticColumnId() { return TEXT("OutTime"); }
+
+protected:
+	//~ Begin INavigationToolColumn
+	virtual FName GetColumnId() const override { return StaticColumnId(); }
+	SEQUENCENAVIGATOR_API virtual FText GetColumnDisplayNameText() const override;
+	SEQUENCENAVIGATOR_API virtual const FSlateBrush* GetIconBrush() const override;
+	virtual bool ShouldShowColumnByDefault() const override { return true; }
+	virtual float GetFillWidth() const override { return 5.f; }
+	SEQUENCENAVIGATOR_API virtual SHeaderRow::FColumn::FArguments ConstructHeaderRowColumn(const TSharedRef<INavigationToolView>& InView
+		, const float InFillSize) override;
+	SEQUENCENAVIGATOR_API virtual TSharedRef<SWidget> ConstructRowWidget(const FNavigationToolItemRef& InItem
+		, const TSharedRef<INavigationToolView>& InView
+		, const TSharedRef<SNavigationToolTreeRow>& InRow) override;
+	//~ End INavigationToolColumn
+};
+
+}
diff --git a/Engine/Plugins/Experimental/SequenceNavigator/Source/SequenceNavigator/Public/Columns/NavigationToolPlayheadColumn.h b/Engine/Plugins/Experimental/SequenceNavigator/Source/SequenceNavigator/Public/Columns/NavigationToolPlayheadColumn.h
new file mode 100644
--- /dev/null
+++ b/Engine/Plugins/Experimental/SequenceNavigator/Source/SequenceNavigator/Public/Columns/NavigationToolPlayheadColumn.h
@@ -0,0 +1,33 @@
+// Copyright Epic Games, Inc. All Rights Reserved.
+
+#pragma once
+
+#include "Columns/NavigationToolColumn.h"
+
+namespace UE::Sequencer::NavigationTool
+{
+
+class INavigationToolView;
+
+class FNavigationToolPlayheadColumn : public FNavigationToolColumn
+{
+public:
+	UE_NAVIGATIONTOOL_INHERITS(FNavigationToolPlayheadColumn, FNavigationToolColumn);
+
+	static FName StaticColumnId() { return TEXT("Playhead"); }
+
+protected:
+	//~ Begin INavigationToolColumn
+	virtual FName GetColumnId() const override { return StaticColumnId(); }
+	SEQUENCENAVIGATOR_API virtual FText GetColumnDisplayNameText() const override;
+	SEQUENCENAVIGATOR_API virtual const FSlateBrush* GetIconBrush() const override;
+	SEQUENCENAVIGATOR_API virtual SHeaderRow::FColumn::FArguments ConstructHeaderRowColumn(const TSharedRef<INavigationToolView>& InView
+		, const float InFillSize) override;
+	virtual bool ShouldShowColumnByDefault() const override { return true; }
+	SEQUENCENAVIGATOR_API virtual TSharedRef<SWidget> ConstructRowWidget(const FNavigationToolItemRef& InItem
+		, const TSharedRef<INavigationToolView>& InView
+		, const TSharedRef<SNavigationToolTreeRow>& InRow) override;
+	//~ End INavigationToolColumn
+};
+
+}
diff --git a/Engine/Plugins/Experimental/SequenceNavigator/Source/SequenceNavigator/Public/Columns/NavigationToolRevisionControlColumn.h b/Engine/Plugins/Experimental/SequenceNavigator/Source/SequenceNavigator/Public/Columns/NavigationToolRevisionControlColumn.h
new file mode 100644
--- /dev/null
+++ b/Engine/Plugins/Experimental/SequenceNavigator/Source/SequenceNavigator/Public/Columns/NavigationToolRevisionControlColumn.h
@@ -0,0 +1,33 @@
+// Copyright Epic Games, Inc. All Rights Reserved.
+
+#pragma once
+
+#include "Columns/NavigationToolColumn.h"
+
+namespace UE::Sequencer::NavigationTool
+{
+
+class INavigationToolView;
+
+class FNavigationToolRevisionControlColumn : public FNavigationToolColumn
+{
+public:
+	UE_NAVIGATIONTOOL_INHERITS(FNavigationToolRevisionControlColumn, FNavigationToolColumn);
+
+	static FName StaticColumnId() { return TEXT("RevisionControl"); }
+
+protected:
+	//~ Begin INavigationToolColumn
+	virtual FName GetColumnId() const override { return StaticColumnId(); }
+	SEQUENCENAVIGATOR_API virtual FText GetColumnDisplayNameText() const override;
+	SEQUENCENAVIGATOR_API virtual const FSlateBrush* GetIconBrush() const override;
+	SEQUENCENAVIGATOR_API virtual SHeaderRow::FColumn::FArguments ConstructHeaderRowColumn(const TSharedRef<INavigationToolView>& InView
+		, const float InFillSize) override;
+	virtual bool ShouldShowColumnByDefault() const override { return true; }
+	SEQUENCENAVIGATOR_API virtual TSharedRef<SWidget> ConstructRowWidget(const FNavigationToolItemRef& InItem
+		, const TSharedRef<INavigationToolView>& InView
+		, const TSharedRef<SNavigationToolTreeRow>& InRow) override;
+	//~ End INavigationToolColumn
+};
+
+}
diff --git a/Engine/Plugins/Experimental/SequenceNavigator/Source/SequenceNavigator/Public/Columns/NavigationToolStartFrameOffsetColumn.h b/Engine/Plugins/Experimental/SequenceNavigator/Source/SequenceNavigator/Public/Columns/NavigationToolStartFrameOffsetColumn.h
new file mode 100644
--- /dev/null
+++ b/Engine/Plugins/Experimental/SequenceNavigator/Source/SequenceNavigator/Public/Columns/NavigationToolStartFrameOffsetColumn.h
@@ -0,0 +1,34 @@
+// Copyright Epic Games, Inc. All Rights Reserved.
+
+#pragma once
+
+#include "Columns/NavigationToolColumn.h"
+
+namespace UE::Sequencer::NavigationTool
+{
+
+class INavigationToolView;
+
+class FNavigationToolStartFrameOffsetColumn : public FNavigationToolColumn
+{
+public:
+	UE_NAVIGATIONTOOL_INHERITS(FNavigationToolStartFrameOffsetColumn, FNavigationToolColumn);
+
+	static FName StaticColumnId() { return TEXT("StartFrameOffset"); }
+
+protected:
+	//~ Begin INavigationToolColumn
+	virtual FName GetColumnId() const override { return StaticColumnId(); }
+	SEQUENCENAVIGATOR_API virtual FText GetColumnDisplayNameText() const override;
+	SEQUENCENAVIGATOR_API virtual const FSlateBrush* GetIconBrush() const override;
+	virtual bool ShouldShowColumnByDefault() const override { return false; }
+	virtual float GetFillWidth() const override { return 5.f; }
+	SEQUENCENAVIGATOR_API virtual SHeaderRow::FColumn::FArguments ConstructHeaderRowColumn(const TSharedRef<INavigationToolView>& InView
+		, const float InFillSize) override;
+	SEQUENCENAVIGATOR_API virtual TSharedRef<SWidget> ConstructRowWidget(const FNavigationToolItemRef& InItem
+		, const TSharedRef<INavigationToolView>& InView
+		, const TSharedRef<SNavigationToolTreeRow>& InRow) override;
+	//~ End INavigationToolColumn
+};
+
+}
diff --git a/Engine/Plugins/Experimental/SequenceNavigator/Source/SequenceNavigator/Public/Columns/NavigationToolTakeColumn.h b/Engine/Plugins/Experimental/SequenceNavigator/Source/SequenceNavigator/Public/Columns/NavigationToolTakeColumn.h
new file mode 100644
--- /dev/null
+++ b/Engine/Plugins/Experimental/SequenceNavigator/Source/SequenceNavigator/Public/Columns/NavigationToolTakeColumn.h
@@ -0,0 +1,34 @@
+// Copyright Epic Games, Inc. All Rights Reserved.
+
+#pragma once
+
+#include "Columns/NavigationToolColumn.h"
+
+namespace UE::Sequencer::NavigationTool
+{
+
+class INavigationToolView;
+
+class FNavigationToolTakeColumn : public FNavigationToolColumn
+{
+public:
+	UE_NAVIGATIONTOOL_INHERITS(FNavigationToolTakeColumn, FNavigationToolColumn);
+
+	static FName StaticColumnId() { return TEXT("ShotTake"); }
+
+protected:
+	//~ Begin INavigationToolColumn
+	virtual FName GetColumnId() const override { return StaticColumnId(); }
+	SEQUENCENAVIGATOR_API virtual FText GetColumnDisplayNameText() const override;
+	SEQUENCENAVIGATOR_API virtual const FSlateBrush* GetIconBrush() const override;
+	virtual bool ShouldShowColumnByDefault() const override { return false; }
+	virtual float GetFillWidth() const override { return 5.f; }
+	SEQUENCENAVIGATOR_API virtual SHeaderRow::FColumn::FArguments ConstructHeaderRowColumn(const TSharedRef<INavigationToolView>& InView
+		, const float InFillSize) override;
+	SEQUENCENAVIGATOR_API virtual TSharedRef<SWidget> ConstructRowWidget(const FNavigationToolItemRef& InItem
+		, const TSharedRef<INavigationToolView>& InView
+		, const TSharedRef<SNavigationToolTreeRow>& InRow) override;
+	//~ End INavigationToolColumn
+};
+
+}
diff --git a/Engine/Plugins/Experimental/SequenceNavigator/Source/SequenceNavigator/Public/DragDropOps/Handlers/NavigationToolItemDropHandler.h b/Engine/Plugins/Experimental/SequenceNavigator/Source/SequenceNavigator/Public/DragDropOps/Handlers/NavigationToolItemDropHandler.h
new file mode 100644
--- /dev/null
+++ b/Engine/Plugins/Experimental/SequenceNavigator/Source/SequenceNavigator/Public/DragDropOps/Handlers/NavigationToolItemDropHandler.h
@@ -0,0 +1,76 @@
+﻿// Copyright Epic Games, Inc. All Rights Reserved.
+
+#pragma once
+
+#include "Containers/Array.h"
+#include "Items/INavigationToolItem.h"
+#include "NavigationToolDefines.h"
+#include "NavigationToolItemType.h"
+#include "Templates/EnableIf.h"
+#include "Templates/SharedPointer.h"
+
+enum class EItemDropZone;
+template<typename OptionalType> struct TOptional;
+
+namespace UE::Sequencer::NavigationTool
+{
+
+class FNavigationToolItemDragDropOp;
+
+/**
+ * Base Class to Handle Dropping Navigation Tool Items into a Target Navigation Tool Item
+ * @see built-in example FNavigationToolActorDropHandler
+ */
+class FNavigationToolItemDropHandler : public INavigationToolItemTypeCastable, public TSharedFromThis<FNavigationToolItemDropHandler>
+{
+	friend FNavigationToolItemDragDropOp;
+
+	SEQUENCENAVIGATOR_API void Initialize(const FNavigationToolItemDragDropOp& InDragDropOp);
+
+public:
+	UE_NAVIGATIONTOOL_INHERITS(FNavigationToolItemDropHandler, INavigationToolItemTypeCastable);
+
+	TConstArrayView<TSharedPtr<INavigationToolItem>> GetItems() const
+	{
+		return Items;
+	}	
+
+protected:
+	virtual bool IsDraggedItemSupported(const FNavigationToolItemPtr& InDraggedItem) const = 0;
+
+	virtual TOptional<EItemDropZone> CanDrop(EItemDropZone InDropZone, const FNavigationToolItemPtr& InTargetItem) const = 0;
+
+	virtual bool Drop(EItemDropZone InDropZone, const FNavigationToolItemPtr& InTargetItem) = 0;
+
+	enum class EIterationResult
+	{
+		Continue,
+		Break,
+	};
+	template<typename InItemType, typename = typename TEnableIf<TIsDerivedFrom<InItemType, INavigationToolItem>::Value>::Type>
+	void ForEachItem(const TFunctionRef<EIterationResult(InItemType&)>& InFunc) const
+	{
+		for (const TSharedPtr<INavigationToolItem>& Item : Items)
+		{
+			if (!Item.IsValid())
+			{
+				continue;
+			}
+
+			if (InItemType* const CastedItem = Item->CastTo<InItemType>())
+			{
+				EIterationResult IterationResult = InFunc(*CastedItem);
+				if (IterationResult == EIterationResult::Break)
+				{
+					break;
+				}
+			}
+		}
+	}
+
+	TArray<FNavigationToolItemPtr> Items;
+
+	ENavigationToolDragDropActionType ActionType = ENavigationToolDragDropActionType::Move;
+};
+
+}
diff --git a/Engine/Plugins/Experimental/SequenceNavigator/Source/SequenceNavigator/Public/DragDropOps/NavigationToolItemDragDropOp.h b/Engine/Plugins/Experimental/SequenceNavigator/Source/SequenceNavigator/Public/DragDropOps/NavigationToolItemDragDropOp.h
new file mode 100644
--- /dev/null
+++ b/Engine/Plugins/Experimental/SequenceNavigator/Source/SequenceNavigator/Public/DragDropOps/NavigationToolItemDragDropOp.h
@@ -0,0 +1,70 @@
+﻿// Copyright Epic Games, Inc. All Rights Reserved.
+
+#pragma once
+
+#include "DragAndDrop/DecoratedDragDropOp.h"
+#include "Handlers/NavigationToolItemDropHandler.h"
+
+enum class EItemDropZone;
+
+namespace UE::Sequencer::NavigationTool
+{
+
+class INavigationToolView;
+class FNavigationToolView;
+class FNavigationToolItemDropHandler;
+
+/** Drag Drop Operation for Navigation Tool Items. Customized behavior can be added in via the AddDropHandler function */
+class FNavigationToolItemDragDropOp : public FDecoratedDragDropOp
+{
+public:
+	DRAG_DROP_OPERATOR_TYPE(FNavigationToolItemDragDropOp, FDecoratedDragDropOp)
+
+	static TSharedRef<FNavigationToolItemDragDropOp> New(const TArray<FNavigationToolItemPtr>& InItems
+		, const TSharedPtr<FNavigationToolView>& InToolView
+		, ENavigationToolDragDropActionType InActionType);
+
+	TSharedPtr<INavigationToolView> GetToolView() const
+	{
+		return WeakToolView.Pin();
+	}
+
+	TConstArrayView<FNavigationToolItemPtr> GetItems() const
+	{
+		return Items;
+	}
+
+	ENavigationToolDragDropActionType GetActionType() const
+	{
+		return ActionType;
+	}
+
+	template<typename InDropHandlerType
+		, typename = typename TEnableIf<TIsDerivedFrom<InDropHandlerType, FNavigationToolItemDropHandler>::Value>::Type
+		, typename... InArgTypes>
+	void AddDropHandler(InArgTypes&&... InArgs)
+	{
+		TSharedRef<FNavigationToolItemDropHandler> DropHandler = MakeShared<InDropHandlerType>(Forward<InArgTypes>(InArgs)...);
+		DropHandler->Initialize(*this);
+		DropHandlers.Add(DropHandler);
+	}
+
+	FReply Drop(EItemDropZone InDropZone, const FNavigationToolItemPtr& InTargetItem);
+
+	TOptional<EItemDropZone> CanDrop(const EItemDropZone InDropZone, const FNavigationToolItemPtr& InTargetItem) const;
+
+protected:
+	void Init(const TArray<FNavigationToolItemPtr>& InItems
+		, const TSharedPtr<FNavigationToolView>& InToolView
+		, ENavigationToolDragDropActionType InActionType);
+
+	TArray<FNavigationToolItemPtr> Items;
+
+	TArray<TSharedRef<FNavigationToolItemDropHandler>> DropHandlers;
+
+	TWeakPtr<INavigationToolView> WeakToolView;
+
+	ENavigationToolDragDropActionType ActionType = ENavigationToolDragDropActionType::Move;
+};
+
+}
diff --git a/Engine/Plugins/Experimental/SequenceNavigator/Source/SequenceNavigator/Public/Extensions/IColorExtension.h b/Engine/Plugins/Experimental/SequenceNavigator/Source/SequenceNavigator/Public/Extensions/IColorExtension.h
new file mode 100644
--- /dev/null
+++ b/Engine/Plugins/Experimental/SequenceNavigator/Source/SequenceNavigator/Public/Extensions/IColorExtension.h
@@ -0,0 +1,21 @@
+﻿// Copyright Epic Games, Inc. All Rights Reserved.
+
+#pragma once
+
+#include "Misc/Optional.h"
+#include "NavigationToolItemType.h"
+
+namespace UE::Sequencer::NavigationTool
+{
+
+class SEQUENCENAVIGATOR_API IColorExtension : public INavigationToolItemTypeCastable
+{
+public:
+	UE_NAVIGATIONTOOL_INHERITS(IColorExtension, INavigationToolItemTypeCastable)
+
+	virtual TOptional<FColor> GetColor() const = 0;
+
+	virtual void SetColor(const TOptional<FColor>& InColor) = 0;
+};
+
+}
diff --git a/Engine/Plugins/Experimental/SequenceNavigator/Source/SequenceNavigator/Public/Extensions/IInTimeExtension.h b/Engine/Plugins/Experimental/SequenceNavigator/Source/SequenceNavigator/Public/Extensions/IInTimeExtension.h
new file mode 100644
--- /dev/null
+++ b/Engine/Plugins/Experimental/SequenceNavigator/Source/SequenceNavigator/Public/Extensions/IInTimeExtension.h
@@ -0,0 +1,21 @@
+﻿// Copyright Epic Games, Inc. All Rights Reserved.
+
+#pragma once
+
+#include "Misc/FrameNumber.h"
+#include "NavigationToolItemType.h"
+
+namespace UE::Sequencer::NavigationTool
+{
+
+class SEQUENCENAVIGATOR_API IInTimeExtension : public INavigationToolItemTypeCastable
+{
+public:
+	UE_NAVIGATIONTOOL_INHERITS(IInTimeExtension, INavigationToolItemTypeCastable)
+
+	virtual FFrameNumber GetInTime() const = 0;
+
+	virtual void SetInTime(const FFrameNumber& InTime) = 0;
+};
+
+}
diff --git a/Engine/Plugins/Experimental/SequenceNavigator/Source/SequenceNavigator/Public/Extensions/IMarkerVisibilityExtension.h b/Engine/Plugins/Experimental/SequenceNavigator/Source/SequenceNavigator/Public/Extensions/IMarkerVisibilityExtension.h
new file mode 100644
--- /dev/null
+++ b/Engine/Plugins/Experimental/SequenceNavigator/Source/SequenceNavigator/Public/Extensions/IMarkerVisibilityExtension.h
@@ -0,0 +1,28 @@
+﻿// Copyright Epic Games, Inc. All Rights Reserved.
+
+#pragma once
+
+#include "NavigationToolItemType.h"
+
+namespace UE::Sequencer::NavigationTool
+{
+
+enum class EItemMarkerVisibility
+{
+	None             = 0,
+	Visible          = 1 << 0,
+	PartiallyVisible = 1 << 1,
+};
+ENUM_CLASS_FLAGS(EItemMarkerVisibility)
+
+class SEQUENCENAVIGATOR_API IMarkerVisibilityExtension : public INavigationToolItemTypeCastable
+{
+public:
+	UE_NAVIGATIONTOOL_INHERITS(IMarkerVisibilityExtension, INavigationToolItemTypeCastable)
+
+	virtual EItemMarkerVisibility GetMarkerVisibility() const = 0;
+
+	virtual void SetMarkerVisibility(const bool bInVisible) = 0;
+};
+
+}
diff --git a/Engine/Plugins/Experimental/SequenceNavigator/Source/SequenceNavigator/Public/Extensions/IOutTimeExtension.h b/Engine/Plugins/Experimental/SequenceNavigator/Source/SequenceNavigator/Public/Extensions/IOutTimeExtension.h
new file mode 100644
--- /dev/null
+++ b/Engine/Plugins/Experimental/SequenceNavigator/Source/SequenceNavigator/Public/Extensions/IOutTimeExtension.h
@@ -0,0 +1,21 @@
+﻿// Copyright Epic Games, Inc. All Rights Reserved.
+
+#pragma once
+
+#include "Misc/FrameNumber.h"
+#include "NavigationToolItemType.h"
+
+namespace UE::Sequencer::NavigationTool
+{
+
+class SEQUENCENAVIGATOR_API IOutTimeExtension : public INavigationToolItemTypeCastable
+{
+public:
+	UE_NAVIGATIONTOOL_INHERITS(IOutTimeExtension, INavigationToolItemTypeCastable)
+
+	virtual FFrameNumber GetOutTime() const = 0;
+
+	virtual void SetOutTime(const FFrameNumber& InTime) = 0;
+};
+
+}
diff --git a/Engine/Plugins/Experimental/SequenceNavigator/Source/SequenceNavigator/Public/Extensions/IPlayheadExtension.h b/Engine/Plugins/Experimental/SequenceNavigator/Source/SequenceNavigator/Public/Extensions/IPlayheadExtension.h
new file mode 100644
--- /dev/null
+++ b/Engine/Plugins/Experimental/SequenceNavigator/Source/SequenceNavigator/Public/Extensions/IPlayheadExtension.h
@@ -0,0 +1,26 @@
+﻿// Copyright Epic Games, Inc. All Rights Reserved.
+
+#pragma once
+
+#include "NavigationToolItemType.h"
+
+namespace UE::Sequencer::NavigationTool
+{
+
+enum class EItemContainsPlayhead
+{
+	None                      = 0,
+	ContainsPlayhead          = 1 << 0,
+	PartiallyContainsPlayhead = 1 << 1,
+};
+ENUM_CLASS_FLAGS(EItemContainsPlayhead)
+
+class SEQUENCENAVIGATOR_API IPlayheadExtension : public INavigationToolItemTypeCastable
+{
+public:
+	UE_NAVIGATIONTOOL_INHERITS(IPlayheadExtension, INavigationToolItemTypeCastable)
+
+	virtual EItemContainsPlayhead ContainsPlayhead() const = 0;
+};
+
+}
diff --git a/Engine/Plugins/Experimental/SequenceNavigator/Source/SequenceNavigator/Public/Extensions/IRenameableExtension.h b/Engine/Plugins/Experimental/SequenceNavigator/Source/SequenceNavigator/Public/Extensions/IRenameableExtension.h
new file mode 100644
--- /dev/null
+++ b/Engine/Plugins/Experimental/SequenceNavigator/Source/SequenceNavigator/Public/Extensions/IRenameableExtension.h
@@ -0,0 +1,22 @@
+﻿// Copyright Epic Games, Inc. All Rights Reserved.
+
+#pragma once
+
+#include "NavigationToolItemType.h"
+
+namespace UE::Sequencer::NavigationTool
+{
+
+class SEQUENCENAVIGATOR_API IRenameableExtension : public INavigationToolItemTypeCastable
+{
+public:
+	UE_NAVIGATIONTOOL_INHERITS(IRenameableExtension, INavigationToolItemTypeCastable)
+
+	/** @return True if this item can be renamed */
+	virtual bool CanRename() const = 0;
+
+	/** The implementation to rename the item */
+	virtual bool Rename(const FString& InName) = 0;
+};
+
+}
