[SequenceNavigator] Rename all front facing text to "Sequence Navigator" from "Navigation Tool"

diff --git a/Engine/Plugins/Experimental/SequenceNavigator/Source/SequenceNavigator/Private/DragDropOps/NavigationToolItemDragDropOp.cpp b/Engine/Plugins/Experimental/SequenceNavigator/Source/SequenceNavigator/Private/DragDropOps/NavigationToolItemDragDropOp.cpp
--- a/Engine/Plugins/Experimental/SequenceNavigator/Source/SequenceNavigator/Private/DragDropOps/NavigationToolItemDragDropOp.cpp
+++ b/Engine/Plugins/Experimental/SequenceNavigator/Source/SequenceNavigator/Private/DragDropOps/NavigationToolItemDragDropOp.cpp
@@ -58,7 +58,7 @@ FReply FNavigationToolItemDragDropOp::Drop(EItemDropZone InDropZone, const FNavi
 {
 	FReply Reply = FReply::Unhandled();
 
-	FScopedTransaction Transaction(LOCTEXT("DropItems", "Navigation Tool Drop Items"));
+	FScopedTransaction Transaction(LOCTEXT("DropItems", "Sequence Navigator Drop Items"));
 
 	for (const TSharedRef<FNavigationToolItemDropHandler>& DropHandler : DropHandlers)
 	{
diff --git a/Engine/Plugins/Experimental/SequenceNavigator/Source/SequenceNavigator/Private/Filters/NavigationToolFilterBar.cpp b/Engine/Plugins/Experimental/SequenceNavigator/Source/SequenceNavigator/Private/Filters/NavigationToolFilterBar.cpp
--- a/Engine/Plugins/Experimental/SequenceNavigator/Source/SequenceNavigator/Private/Filters/NavigationToolFilterBar.cpp
+++ b/Engine/Plugins/Experimental/SequenceNavigator/Source/SequenceNavigator/Private/Filters/NavigationToolFilterBar.cpp
@@ -1790,7 +1790,7 @@ void FNavigationToolFilterBar::OpenTextExpressionHelp()
 {
 	FFilterExpressionHelpDialogConfig Config;
 	Config.IdentifierName = TEXT("NavigationToolCustomTextFilterHelp");
-	Config.DialogTitle = LOCTEXT("NavigationToolCustomTextFilterHelp", "Navigation Tool Custom Text Filter Help");
+	Config.DialogTitle = LOCTEXT("CustomTextFilterHelp", "Sequence Navigator Custom Text Filter Help");
 	Config.TextFilterExpressionContexts = GetTextFilterExpressionContexts();
 	SFilterExpressionHelpDialog::Open(MoveTemp(Config));
 }
diff --git a/Engine/Plugins/Experimental/SequenceNavigator/Source/SequenceNavigator/Private/Filters/NavigationToolFilterCommands.h b/Engine/Plugins/Experimental/SequenceNavigator/Source/SequenceNavigator/Private/Filters/NavigationToolFilterCommands.h
--- a/Engine/Plugins/Experimental/SequenceNavigator/Source/SequenceNavigator/Private/Filters/NavigationToolFilterCommands.h
+++ b/Engine/Plugins/Experimental/SequenceNavigator/Source/SequenceNavigator/Private/Filters/NavigationToolFilterCommands.h
@@ -13,7 +13,7 @@ class FNavigationToolFilterCommands
 public:
 	FNavigationToolFilterCommands()
 		: TCommands<FNavigationToolFilterCommands>(TEXT("SequenceNavigationToolFilters")
-		, LOCTEXT("SequenceNavigationToolFilters", "Sequence Navigation Tool Filters")
+		, LOCTEXT("SequenceNavigatorFilters", "Sequence Navigator Filters")
 		, NAME_None
 		, FAppStyle::GetAppStyleSetName())
 	{}
diff --git a/Engine/Plugins/Experimental/SequenceNavigator/Source/SequenceNavigator/Private/Menus/NavigationToolItemContextMenu.cpp b/Engine/Plugins/Experimental/SequenceNavigator/Source/SequenceNavigator/Private/Menus/NavigationToolItemContextMenu.cpp
--- a/Engine/Plugins/Experimental/SequenceNavigator/Source/SequenceNavigator/Private/Menus/NavigationToolItemContextMenu.cpp
+++ b/Engine/Plugins/Experimental/SequenceNavigator/Source/SequenceNavigator/Private/Menus/NavigationToolItemContextMenu.cpp
@@ -99,7 +99,7 @@ void FNavigationToolItemContextMenu::CreateToolSection(UToolMenu& InMenu)
 {
 	const FNavigationToolCommands& ToolCommands = FNavigationToolCommands::Get();
 
-	FToolMenuSection& ToolSection = InMenu.FindOrAddSection(TEXT("ToolActions"), LOCTEXT("ToolActionsHeader", "Navigation Tool Actions"));
+	FToolMenuSection& ToolSection = InMenu.FindOrAddSection(TEXT("ToolActions"), LOCTEXT("ToolActionsHeader", "Sequence Navigator Actions"));
 
 	FToolMenuEntry& ExpandAllEntry = ToolSection.AddMenuEntry(ToolCommands.ExpandAll);
 	ExpandAllEntry.Icon = FSlateIcon(FAppStyle::GetAppStyleSetName(), TEXT("TreeArrow_Expanded"));
diff --git a/Engine/Plugins/Experimental/SequenceNavigator/Source/SequenceNavigator/Private/NavigationTool.cpp b/Engine/Plugins/Experimental/SequenceNavigator/Source/SequenceNavigator/Private/NavigationTool.cpp
--- a/Engine/Plugins/Experimental/SequenceNavigator/Source/SequenceNavigator/Private/NavigationTool.cpp
+++ b/Engine/Plugins/Experimental/SequenceNavigator/Source/SequenceNavigator/Private/NavigationTool.cpp
@@ -433,7 +433,7 @@ void FNavigationTool::Refresh()
 
 			if (bShouldTransact)
 			{
-				Transaction = MakeUnique<FScopedTransaction>(LOCTEXT("NavigationToolItemAction", "Navigation Tool Item Action"));
+				Transaction = MakeUnique<FScopedTransaction>(LOCTEXT("ItemAction", "Sequence Navigator Item Action"));
 			}
 		}
 
@@ -624,7 +624,7 @@ void FNavigationTool::DeleteItems(TArray<FNavigationToolItemPtr> InItems)
 		return;
 	}
 
-	FScopedTransaction DeleteTransaction(LOCTEXT("NavigationToolItemDeleteAction", "Navigation Tool Delete Item(s)"), !GIsTransacting);
+	FScopedTransaction DeleteTransaction(LOCTEXT("ItemDeleteAction", "Sequence Navigator Delete Item(s)"), !GIsTransacting);
 
 	bool bRequestRefresh = false;
 	for (FNavigationToolItemPtr& Item : InItems)
diff --git a/Engine/Plugins/Experimental/SequenceNavigator/Source/SequenceNavigator/Private/NavigationToolCommands.cpp b/Engine/Plugins/Experimental/SequenceNavigator/Source/SequenceNavigator/Private/NavigationToolCommands.cpp
--- a/Engine/Plugins/Experimental/SequenceNavigator/Source/SequenceNavigator/Private/NavigationToolCommands.cpp
+++ b/Engine/Plugins/Experimental/SequenceNavigator/Source/SequenceNavigator/Private/NavigationToolCommands.cpp
@@ -9,8 +9,8 @@ namespace UE::Sequencer::NavigationTool
 {
 
 FNavigationToolCommands::FNavigationToolCommands()
-	: TCommands<FNavigationToolCommands>(TEXT("SequenceNavigation")
-	, LOCTEXT("SequenceNavigationTool", "Sequence Navigation")
+	: TCommands<FNavigationToolCommands>(TEXT("SequenceNavigator")
+	, LOCTEXT("SequenceNavigator", "Sequence Navigator")
 	, NAME_None
 	, FAppStyle::GetAppStyleSetName())
 {
@@ -21,15 +21,15 @@ void FNavigationToolCommands::RegisterCommands()
 	FUICommandInfo::MakeCommandInfo(AsShared()
 		, OpenToolSettings
 		, TEXT("OpenToolSettings")
-		, LOCTEXT("OpenToolSettings", "Navigation Tool Settings...")
-		, LOCTEXT("OpenToolSettingsTooltip", "Opens the Navigation Tool editor settings")
+		, LOCTEXT("OpenToolSettings", "Sequence Navigator Settings...")
+		, LOCTEXT("OpenToolSettingsTooltip", "Opens the Sequence Navigator editor settings")
 		, FSlateIcon(FAppStyle::Get().GetStyleSetName(), TEXT("Icons.Settings"))
 		, EUserInterfaceActionType::Button
 		, FInputChord());
 
 	UI_COMMAND(ToggleToolTabVisible
-		, "Sequence Navigation"
-		, "Toggle the visibility of the Sequence Navigation Tool"
+		, "Sequence Navigator"
+		, "Toggle the visibility of the Sequence Navigator"
 		, EUserInterfaceActionType::ToggleButton
 		, FInputChord());
 
diff --git a/Engine/Plugins/Experimental/SequenceNavigator/Source/SequenceNavigator/Private/NavigationToolExtender.cpp b/Engine/Plugins/Experimental/SequenceNavigator/Source/SequenceNavigator/Private/NavigationToolExtender.cpp
--- a/Engine/Plugins/Experimental/SequenceNavigator/Source/SequenceNavigator/Private/NavigationToolExtender.cpp
+++ b/Engine/Plugins/Experimental/SequenceNavigator/Source/SequenceNavigator/Private/NavigationToolExtender.cpp
@@ -193,7 +193,7 @@ bool FNavigationToolExtender::RegisterToolProvider(const TSharedRef<ISequencer>&
 	const FName ToolId = GetToolInstanceId(*InSequencer);
 	if (ToolId.IsNone())
 	{
-		UE_LOG(LogSequenceNavigator, Error, TEXT("Invalid Navigation Tool Instance Name"));
+		UE_LOG(LogSequenceNavigator, Error, TEXT("Invalid Sequence Navigator Instance Name"));
 		return false;
 	}
 
@@ -201,13 +201,13 @@ bool FNavigationToolExtender::RegisterToolProvider(const TSharedRef<ISequencer>&
 
 	if (NewProviderId.IsNone())
 	{
-		UE_LOG(LogSequenceNavigator, Error, TEXT("Invalid Navigation Tool Provider Name"));
+		UE_LOG(LogSequenceNavigator, Error, TEXT("Invalid Sequence Navigator Provider Name"));
 		return false;
 	}
 
 	if (const TSharedPtr<FNavigationToolProvider> ExistingProvider = FindToolProvider(ToolId, NewProviderId))
 	{
-		UE_LOG(LogSequenceNavigator, Warning, TEXT("Navigation Tool provider already registered: %s")
+		UE_LOG(LogSequenceNavigator, Warning, TEXT("Sequence Navigator provider already registered: %s")
 			, *NewProviderId.ToString());
 		return true;
 	}
@@ -219,7 +219,7 @@ bool FNavigationToolExtender::RegisterToolProvider(const TSharedRef<ISequencer>&
 	FNavigationToolInstance* const ToolInstance = This.ToolInstances.Find(ToolId);
 	if (!ToolInstance)
 	{
-		UE_LOG(LogSequenceNavigator, Error, TEXT("Navigation Tool instance could not be created for provider: %s")
+		UE_LOG(LogSequenceNavigator, Error, TEXT("Sequence Navigator instance could not be created for provider: %s")
 			, *NewProviderId.ToString());
 		return true;
 	}
@@ -228,7 +228,7 @@ bool FNavigationToolExtender::RegisterToolProvider(const TSharedRef<ISequencer>&
 
 	AddSequencerToolBarExtension(*ToolInstance);
 
-	UE_LOG(LogSequenceNavigator, Log, TEXT("Navigation Tool '%s' Instance Provider '%s' Registered")
+	UE_LOG(LogSequenceNavigator, Log, TEXT("Sequence Navigator '%s' Instance Provider '%s' Registered")
 		, *ToolId.ToString(), *NewProviderId.ToString());
 
 	This.ProvidersChangedDelegate.Broadcast(ToolId, InProvider, ENavigationToolProvidersChangeType::Add);
@@ -253,7 +253,7 @@ bool FNavigationToolExtender::UnregisterToolProvider(const FName InToolId, const
 		const TSharedRef<FNavigationToolProvider>& Provider = *It;
 		if (Provider->GetIdentifier() == InProviderId)
 		{
-			UE_LOG(LogSequenceNavigator, Log, TEXT("Navigation Tool Provider Unregistered: %s")
+			UE_LOG(LogSequenceNavigator, Log, TEXT("Sequence Navigator Provider Unregistered: %s")
 				, *Provider->GetIdentifier().ToString());
 
 			This.ProvidersChangedDelegate.Broadcast(InToolId, Provider, ENavigationToolProvidersChangeType::Remove);
diff --git a/Engine/Plugins/Experimental/SequenceNavigator/Source/SequenceNavigator/Private/NavigationToolView.cpp b/Engine/Plugins/Experimental/SequenceNavigator/Source/SequenceNavigator/Private/NavigationToolView.cpp
--- a/Engine/Plugins/Experimental/SequenceNavigator/Source/SequenceNavigator/Private/NavigationToolView.cpp
+++ b/Engine/Plugins/Experimental/SequenceNavigator/Source/SequenceNavigator/Private/NavigationToolView.cpp
@@ -451,7 +451,7 @@ void FNavigationToolView::SaveViewItemFlags(const FNavigationToolItemPtr& InItem
 	FNavigationToolSaveState* const SaveState = Provider->GetSaveState(*Tool);
 	if (!SaveState)
 	{
-		UE_LOG(LogSequenceNavigator, Warning, TEXT("SaveViewItemFlags(): Navigation Tool Save State is NULL!"));
+		UE_LOG(LogSequenceNavigator, Warning, TEXT("SaveViewItemFlags(): Sequence Navigator Save State is NULL!"));
 		return;
 	}
 
@@ -523,7 +523,7 @@ void FNavigationToolView::GetChildrenOfItem(const FNavigationToolItemPtr& InItem
 	const TSharedPtr<FNavigationToolProvider> Provider = InItem->GetProvider();
 	if (!Provider.IsValid() && InItem->GetItemId() != FNavigationToolItemId::RootId)
 	{
-		UE_LOG(LogTemp, Warning, TEXT("Navigation Tool Item Id \"%s\" has no provider, but is a root item!"), *InItem->GetItemId().GetStringId());
+		UE_LOG(LogTemp, Warning, TEXT("Sequence Navigator Item Id \"%s\" has no provider, but is a root item!"), *InItem->GetItemId().GetStringId());
 		return;
 	}
 
@@ -884,13 +884,13 @@ bool FNavigationToolView::ShouldShowItem(const FNavigationToolItemPtr& InItem, b
 
 	if (!InItem->IsAllowedInTool())
 	{
-		//UE_LOG(LogTemp, Warning, TEXT("Navigation Tool Item Id \"%s\" is hidden: Not Allowed In Tool"), *InItem->GetItemId().GetStringId());
+		//UE_LOG(LogTemp, Warning, TEXT("Sequence Navigator Item Id \"%s\" is hidden: Not Allowed In Tool"), *InItem->GetItemId().GetStringId());
 		return false;
 	}
 
 	if (!InItem->IsViewModeSupported(InViewMode, *this))
 	{
-		//UE_LOG(LogTemp, Warning, TEXT("Navigation Tool Item Id \"%s\" is hidden: View Mode Not Supported"), *InItem->GetItemId().GetStringId());
+		//UE_LOG(LogTemp, Warning, TEXT("Sequence Navigator Item Id \"%s\" is hidden: View Mode Not Supported"), *InItem->GetItemId().GetStringId());
 		return false;
 	}
 
@@ -908,7 +908,7 @@ bool FNavigationToolView::ShouldShowItem(const FNavigationToolItemPtr& InItem, b
 		});
 	if (bProviderShouldHideItem)
 	{
-		//UE_LOG(LogTemp, Warning, TEXT("Navigation Tool Item Id \"%s\" is hidden: Provider Should Hide Item"), *InItem->GetItemId().GetStringId());
+		//UE_LOG(LogTemp, Warning, TEXT("Sequence Navigator Item Id \"%s\" is hidden: Provider Should Hide Item"), *InItem->GetItemId().GetStringId());
 		return false;
 	}
 
@@ -947,7 +947,7 @@ bool FNavigationToolView::ShouldShowItem(const FNavigationToolItemPtr& InItem, b
 
 	if (bInUseFilters && FilterBar->GetFilterData().IsFilteredOut(InItem))
 	{
-		//UE_LOG(LogTemp, Warning, TEXT("Navigation Tool Item Id \"%s\" is hidden: Filtered Out"), *InItem->GetItemId().GetStringId());
+		//UE_LOG(LogTemp, Warning, TEXT("Sequence Navigator Item Id \"%s\" is hidden: Filtered Out"), *InItem->GetItemId().GetStringId());
 		return false;
 	}
 
diff --git a/Engine/Plugins/Experimental/SequenceNavigator/Source/SequenceNavigator/Private/Providers/NavigationToolProvider.cpp b/Engine/Plugins/Experimental/SequenceNavigator/Source/SequenceNavigator/Private/Providers/NavigationToolProvider.cpp
--- a/Engine/Plugins/Experimental/SequenceNavigator/Source/SequenceNavigator/Private/Providers/NavigationToolProvider.cpp
+++ b/Engine/Plugins/Experimental/SequenceNavigator/Source/SequenceNavigator/Private/Providers/NavigationToolProvider.cpp
@@ -92,7 +92,7 @@ FNavigationToolViewSaveState* FNavigationToolProvider::GetViewSaveState(const IN
 	FNavigationToolSaveState* const SaveState = GetSaveState(InTool);
 	if (!SaveState)
 	{
-		UE_LOG(LogSequenceNavigator, Warning, TEXT("GetViewSaveState(): Navigation Tool Save State is NULL!"));
+		UE_LOG(LogSequenceNavigator, Warning, TEXT("GetViewSaveState(): Sequence Navigator Save State is NULL!"));
 		return nullptr;
 	}
 
@@ -110,7 +110,7 @@ void FNavigationToolProvider::EnsureToolViewCount(const INavigationTool& InTool,
 	FNavigationToolSaveState* const SaveState = GetSaveState(InTool);
 	if (!SaveState)
 	{
-		UE_LOG(LogSequenceNavigator, Warning, TEXT("EnsureToolViewCount(): Navigation Tool Save State is NULL!"));
+		UE_LOG(LogSequenceNavigator, Warning, TEXT("EnsureToolViewCount(): Sequence Navigator Save State is NULL!"));
 		return;
 	}
 
diff --git a/Engine/Plugins/Experimental/SequenceNavigator/Source/SequenceNavigator/Private/Widgets/Columns/SNavigationToolItemChip.cpp b/Engine/Plugins/Experimental/SequenceNavigator/Source/SequenceNavigator/Private/Widgets/Columns/SNavigationToolItemChip.cpp
--- a/Engine/Plugins/Experimental/SequenceNavigator/Source/SequenceNavigator/Private/Widgets/Columns/SNavigationToolItemChip.cpp
+++ b/Engine/Plugins/Experimental/SequenceNavigator/Source/SequenceNavigator/Private/Widgets/Columns/SNavigationToolItemChip.cpp
@@ -289,7 +289,7 @@ FText SNavigationToolItemChip::GetTooltipText() const
 
 	return FText::Format(LOCTEXT("TooltipText", "{0}\n\n"
 		"Click to select in Sequencer\n\n"
-		"Alt + Click to select item in Navigation Tool tree")
+		"Alt + Click to select item in Sequence Navigator")
 		, ItemDisplayName);
 }
 
