[Sequencer] Extend search box toolbar entry to fill available space

diff --git a/Engine/Plugins/Experimental/SequenceNavigator/Source/SequenceNavigator/Private/Menus/NavigationToolToolbarMenu.cpp b/Engine/Plugins/Experimental/SequenceNavigator/Source/SequenceNavigator/Private/Menus/NavigationToolToolbarMenu.cpp
--- a/Engine/Plugins/Experimental/SequenceNavigator/Source/SequenceNavigator/Private/Menus/NavigationToolToolbarMenu.cpp
+++ b/Engine/Plugins/Experimental/SequenceNavigator/Source/SequenceNavigator/Private/Menus/NavigationToolToolbarMenu.cpp
@@ -89,9 +89,13 @@ void FNavigationToolToolbarMenu::PopulateToolBar(UToolMenu* const InMenu)
 		FilterBar->MakeAddFilterButton(),
 		LOCTEXT("AddFilter", "Add Filter")));
 
-	MainSection.AddEntry(FToolMenuEntry::InitWidget(TEXT("SearchBox"),
-		FilterBar->GetOrCreateSearchBoxWidget(),
-		FText()));
+	FToolMenuEntry SearchBoxEntry = FToolMenuEntry::InitWidget(TEXT("SearchBox"),
+		FilterBar->GetOrCreateSearchBoxWidget(), FText());
+	SearchBoxEntry.WidgetData.StyleParams.VerticalAlignment = VAlign_Center;
+	SearchBoxEntry.WidgetData.StyleParams.HorizontalAlignment = HAlign_Fill;
+	SearchBoxEntry.WidgetData.StyleParams.FillSize = 1.f;
+	SearchBoxEntry.WidgetData.StyleParams.SizeRule = FSizeParam::ESizeRule::SizeRule_Stretch;
+	MainSection.AddEntry(SearchBoxEntry);
 
 	const FToolMenuEntry SettingsEntry = FToolMenuEntry::InitComboButton(TEXT("Settings"),
 		FUIAction(),
