[SequenceNavigator] Don't assume UToolMenu::Get() is valid in FNavigationToolExtender destructor (fixes crash on editor close)

diff --git a/Engine/Plugins/Experimental/SequenceNavigator/Source/SequenceNavigator/Private/NavigationToolExtender.cpp b/Engine/Plugins/Experimental/SequenceNavigator/Source/SequenceNavigator/Private/NavigationToolExtender.cpp
--- a/Engine/Plugins/Experimental/SequenceNavigator/Source/SequenceNavigator/Private/NavigationToolExtender.cpp
+++ b/Engine/Plugins/Experimental/SequenceNavigator/Source/SequenceNavigator/Private/NavigationToolExtender.cpp
@@ -51,9 +51,12 @@ FNavigationToolExtender::~FNavigationToolExtender()
 		SequencerModule.UnregisterOnSequencerCreated(SequencerCreatedHandle);
 	}
 
-	for (const TPair<FName, FNavigationToolInstance>& ToolInstance : ToolInstances)
+	if (FModuleManager::Get().IsModuleLoaded(TEXT("ToolMenus")))
 	{
-		RemoveSequencerToolBarExtension(ToolInstance.Value);
+		for (const TPair<FName, FNavigationToolInstance>& ToolInstance : ToolInstances)
+		{
+			RemoveSequencerToolBarExtension(ToolInstance.Value);
+		}
 	}
 }
 
