Sequencer: Remove ensure - subsequence can have an empty movie scene if it doesn't exist

#rnx

diff --git a/Engine/Plugins/Experimental/SequenceNavigator/Source/SequenceNavigator/Private/Items/NavigationToolSubTrack.cpp b/Engine/Plugins/Experimental/SequenceNavigator/Source/SequenceNavigator/Private/Items/NavigationToolSubTrack.cpp
--- a/Engine/Plugins/Experimental/SequenceNavigator/Source/SequenceNavigator/Private/Items/NavigationToolSubTrack.cpp
+++ b/Engine/Plugins/Experimental/SequenceNavigator/Source/SequenceNavigator/Private/Items/NavigationToolSubTrack.cpp
@@ -50,7 +50,7 @@ void FNavigationToolSubTrack::FindChildren(TArray<FNavigationToolItemPtr>& OutCh
 			if (IsValid(SubSection))
 			{
 				UMovieSceneSequence* const Sequence = SubSection->GetSequence();
-				if (ensure(IsValid(Sequence)))
+				if (IsValid(Sequence))
 				{
 					const FNavigationToolItemPtr NewItem = Tool.FindOrAdd<FNavigationToolSequence>(GetProvider(), SharedThisRef, Sequence, SubSection, Index);
 					OutChildren.Add(NewItem);
