[SequenceNavigator] Remove unneeded IsValid() checks

diff --git a/Engine/Plugins/Experimental/SequenceNavigator/Source/SequenceNavigator/Private/Filters/NavigationToolFilterBar.cpp b/Engine/Plugins/Experimental/SequenceNavigator/Source/SequenceNavigator/Private/Filters/NavigationToolFilterBar.cpp
--- a/Engine/Plugins/Experimental/SequenceNavigator/Source/SequenceNavigator/Private/Filters/NavigationToolFilterBar.cpp
+++ b/Engine/Plugins/Experimental/SequenceNavigator/Source/SequenceNavigator/Private/Filters/NavigationToolFilterBar.cpp
@@ -205,15 +205,15 @@ void FNavigationToolFilterBar::BindCommands(const TSharedPtr<FUICommandList>& In
 
 void FNavigationToolFilterBar::CreateCustomTextFiltersFromConfig()
 {
-	UNavigationToolSettings* const Settings = GetMutableDefault<UNavigationToolSettings>();
-	if (!Settings)
+	UNavigationToolSettings* const ToolSettings = GetMutableDefault<UNavigationToolSettings>();
+	if (!ToolSettings)
 	{
 		return;
 	}
 
 	CustomTextFilters.Empty();
 
-	FSequencerFilterBarConfig& Config = Settings->FindOrAddFilterBar(GetIdentifier(), false);
+	FSequencerFilterBarConfig& Config = ToolSettings->FindOrAddFilterBar(GetIdentifier(), false);
 
 	for (const FCustomTextFilterData& CustomTextFilterData : Config.GetCustomTextFilters())
 	{
@@ -491,13 +491,13 @@ bool FNavigationToolFilterBar::HasAnyFilterActive(const bool bCheckTextFilter
 
 bool FNavigationToolFilterBar::IsFilterActive(const TSharedRef<FNavigationToolFilter> InFilter) const
 {
-	UNavigationToolSettings* const Settings = GetMutableDefault<UNavigationToolSettings>();
-	if (!Settings)
+	UNavigationToolSettings* const ToolSettings = GetMutableDefault<UNavigationToolSettings>();
+	if (!ToolSettings)
 	{
 		return false;
 	}
 
-	const FSequencerFilterBarConfig& Config = Settings->FindOrAddFilterBar(GetIdentifier(), false);
+	const FSequencerFilterBarConfig& Config = ToolSettings->FindOrAddFilterBar(GetIdentifier(), false);
 
 	const FString FilterName = InFilter->GetDisplayName().ToString();
 	return Config.IsFilterActive(FilterName);
@@ -505,15 +505,15 @@ bool FNavigationToolFilterBar::IsFilterActive(const TSharedRef<FNavigationToolFi
 
 bool FNavigationToolFilterBar::SetFilterActive(const TSharedRef<FNavigationToolFilter>& InFilter, const bool bInActive, const bool bInRequestFilterUpdate)
 {
-	UNavigationToolSettings* const Settings = GetMutableDefault<UNavigationToolSettings>();
-	if (!ensure(IsValid(Settings)))
+	UNavigationToolSettings* const ToolSettings = GetMutableDefault<UNavigationToolSettings>();
+	if (!ToolSettings)
 	{
 		return false;
 	}
 
 	const bool bNewActive = InFilter->IsInverseFilter() ? !bInActive : bInActive;
 
-	FSequencerFilterBarConfig& Config = Settings->FindOrAddFilterBar(GetIdentifier(), true);
+	FSequencerFilterBarConfig& Config = ToolSettings->FindOrAddFilterBar(GetIdentifier(), true);
 
 	const FString FilterName = InFilter->GetDisplayName().ToString();
 	const bool bSuccess = (Config.IsFilterActive(FilterName) == bInActive)
@@ -521,7 +521,7 @@ bool FNavigationToolFilterBar::SetFilterActive(const TSharedRef<FNavigationToolF
 
 	if (bSuccess)
 	{
-		Settings->SaveConfig();
+		ToolSettings->SaveConfig();
 
 		InFilter->SetActive(bNewActive);
 		InFilter->ActiveStateChanged(bNewActive);
@@ -579,13 +579,13 @@ void FNavigationToolFilterBar::ActivateCommonFilters(const bool bInActivate
     , const TArray<TSharedRef<FFilterCategory>>& InMatchCategories
     , const TArray<TSharedRef<FNavigationToolFilter>>& InExceptions)
 {
-	UNavigationToolSettings* const Settings = GetMutableDefault<UNavigationToolSettings>();
-	if (!ensure(IsValid(Settings)))
+	UNavigationToolSettings* const ToolSettings = GetMutableDefault<UNavigationToolSettings>();
+	if (!ToolSettings)
 	{
 		return;
 	}
 
-	FSequencerFilterBarConfig& Config = Settings->FindOrAddFilterBar(GetIdentifier(), false);
+	FSequencerFilterBarConfig& Config = ToolSettings->FindOrAddFilterBar(GetIdentifier(), false);
 
 	bool bNeedsSave = false;
 
@@ -615,7 +615,7 @@ void FNavigationToolFilterBar::ActivateCommonFilters(const bool bInActivate
 
 	if (bNeedsSave)
 	{
-		Settings->SaveConfig();
+		ToolSettings->SaveConfig();
 	}
 
 	RequestFilterUpdate();
@@ -775,13 +775,13 @@ bool FNavigationToolFilterBar::HasAnyFilterEnabled() const
 
 bool FNavigationToolFilterBar::IsFilterEnabled(const TSharedRef<FNavigationToolFilter> InFilter) const
 {
-	UNavigationToolSettings* const Settings = GetMutableDefault<UNavigationToolSettings>();
-	if (!ensure(IsValid(Settings)))
+	UNavigationToolSettings* const ToolSettings = GetMutableDefault<UNavigationToolSettings>();
+	if (!ToolSettings)
 	{
 		return false;
 	}
 
-	FSequencerFilterBarConfig& Config = Settings->FindOrAddFilterBar(GetIdentifier(), false);
+	FSequencerFilterBarConfig& Config = ToolSettings->FindOrAddFilterBar(GetIdentifier(), false);
 
 	const FString FilterName = InFilter->GetDisplayName().ToString();
 	return Config.IsFilterEnabled(FilterName);
@@ -796,7 +796,7 @@ bool FNavigationToolFilterBar::SetFilterEnabled(const TSharedRef<FNavigationTool
 	}
 
 	UNavigationToolSettings* const ToolSettings = GetMutableDefault<UNavigationToolSettings>();
-	if (!ensure(IsValid(ToolSettings)))
+	if (!ToolSettings)
 	{
 		return false;
 	}
@@ -840,7 +840,7 @@ void FNavigationToolFilterBar::EnableFilters(const bool bInEnable
 	}
 
 	UNavigationToolSettings* const ToolSettings = GetMutableDefault<UNavigationToolSettings>();
-	if (!ensure(IsValid(ToolSettings)))
+	if (!ToolSettings)
 	{
 		return;
 	}
@@ -993,14 +993,13 @@ bool FNavigationToolFilterBar::AddCustomTextFilter(const TSharedRef<FNavigationT
 
 	if (bInAddToConfig)
 	{
-		UNavigationToolSettings* const Settings = GetMutableDefault<UNavigationToolSettings>();
-		if (ensure(IsValid(Settings)))
+		if (UNavigationToolSettings* const ToolSettings = GetMutableDefault<UNavigationToolSettings>())
 		{
-			FSequencerFilterBarConfig& Config = Settings->FindOrAddFilterBar(GetIdentifier(), false);
+			FSequencerFilterBarConfig& Config = ToolSettings->FindOrAddFilterBar(GetIdentifier(), false);
 
 			if (Config.AddCustomTextFilter(InFilter->CreateCustomTextFilterData()))
 			{
-				Settings->SaveConfig();
+				ToolSettings->SaveConfig();
 			}
 		}
 	}
@@ -1019,15 +1018,14 @@ bool FNavigationToolFilterBar::RemoveCustomTextFilter(const TSharedRef<FNavigati
 
 	if (bInRemoveFromConfig)
 	{
-		UNavigationToolSettings* const Settings = GetMutableDefault<UNavigationToolSettings>();
-		if (ensure(IsValid(Settings)))
+		if (UNavigationToolSettings* const ToolSettings = GetMutableDefault<UNavigationToolSettings>())
 		{
-			FSequencerFilterBarConfig& Config = Settings->FindOrAddFilterBar(GetIdentifier(), false);
+			FSequencerFilterBarConfig& Config = ToolSettings->FindOrAddFilterBar(GetIdentifier(), false);
 
 			const FString FilterName = InFilter->GetDisplayName().ToString();
 			if (Config.RemoveCustomTextFilter(FilterName))
 			{
-				Settings->SaveConfig();
+				ToolSettings->SaveConfig();
 			}
 		}
 	}
@@ -1046,7 +1044,7 @@ void FNavigationToolFilterBar::ActivateCustomTextFilters(const bool bInActivate,
 	}
 
 	UNavigationToolSettings* const ToolSettings = GetMutableDefault<UNavigationToolSettings>();
-	if (!ensure(IsValid(ToolSettings)))
+	if (!ToolSettings)
 	{
 		return;
 	}
@@ -1086,7 +1084,7 @@ void FNavigationToolFilterBar::EnableCustomTextFilters(const bool bInEnable, con
 	}
 
 	UNavigationToolSettings* const ToolSettings = GetMutableDefault<UNavigationToolSettings>();
-	if (!ensure(IsValid(ToolSettings)))
+	if (!ToolSettings)
 	{
 		return;
 	}
@@ -1167,7 +1165,7 @@ bool FNavigationToolFilterBar::PassesAnyCommonFilter(const FNavigationToolItemPt
 	}
 
 	UMovieSceneSequence* const FocusedSequence = Sequencer->GetFocusedMovieSceneSequence();
-	if (!IsValid(FocusedSequence))
+	if (!FocusedSequence)
 	{
 		return true;
 	}
@@ -1234,7 +1232,7 @@ const FNavigationToolFilterData& FNavigationToolFilterBar::FilterNodes()
 	}
 
 	const UNavigationToolSettings* const ToolSettings = GetDefault<UNavigationToolSettings>();
-	if (!ensure(IsValid(ToolSettings)))
+	if (!ToolSettings)
 	{
 		return FilterData;
 	}
@@ -1246,7 +1244,7 @@ const FNavigationToolFilterData& FNavigationToolFilterBar::FilterNodes()
 	const bool bHasActiveFilter = HasAnyFilterActive();
 	for (const FNavigationToolItemPtr& RootNode : Tool.GetTreeRoot()->GetChildren())
 	{
-		FilterNodesRecursive(*Sequencer,*ToolView, *ToolSettings, bHasActiveFilter, RootNode);
+		FilterNodesRecursive(*ToolView, *ToolSettings, bHasActiveFilter, RootNode);
 	}
 
 	return FilterData;
@@ -1267,8 +1265,7 @@ TArray<TSharedRef<ISequencerTextFilterExpressionContext>> FNavigationToolFilterB
 	return TextFilter->GetTextFilterExpressionContexts();
 }
 
-bool FNavigationToolFilterBar::FilterNodesRecursive(ISequencer& InSequencer
-	, INavigationToolView& InToolView
+bool FNavigationToolFilterBar::FilterNodesRecursive(INavigationToolView& InToolView
 	, const UNavigationToolSettings& InSettings
 	, const bool bInHasActiveFilter
 	, const FNavigationToolItemPtr& InStartNode)
@@ -1286,7 +1283,7 @@ bool FNavigationToolFilterBar::FilterNodesRecursive(ISequencer& InSequencer
 	// Child nodes should always be processed, as they may force their parents to pass
 	for (const FNavigationToolItemPtr& Node : InStartNode->GetChildren())
 	{
-		if (FilterNodesRecursive(InSequencer, InToolView, InSettings, bInHasActiveFilter, Node))
+		if (FilterNodesRecursive(InToolView, InSettings, bInHasActiveFilter, Node))
 		{
 			bAnyChildPassed = true;
 		}
@@ -1454,73 +1451,70 @@ TSharedRef<SComboButton> FNavigationToolFilterBar::MakeAddFilterButton()
 
 bool FNavigationToolFilterBar::ShouldShowFilterBarWidget() const
 {
-	const UNavigationToolSettings* const Settings = GetDefault<UNavigationToolSettings>();
-	if (ensure(IsValid(Settings)))
+	if (const UNavigationToolSettings* const ToolSettings = GetDefault<UNavigationToolSettings>())
 	{
-		return Settings->IsFilterBarVisible() && HasAnyFiltersEnabled();
+		return ToolSettings->IsFilterBarVisible() && HasAnyFiltersEnabled();
 	}
 	return false;
 }
 
 bool FNavigationToolFilterBar::IsFilterBarVisible() const
 {
-	const UNavigationToolSettings* const Settings = GetDefault<UNavigationToolSettings>();
-	if (ensure(IsValid(Settings)))
+	if (const UNavigationToolSettings* const ToolSettings = GetDefault<UNavigationToolSettings>())
 	{
-		return Settings->IsFilterBarVisible();
+		return ToolSettings->IsFilterBarVisible();
 	}
 	return false;
 }
 
 void FNavigationToolFilterBar::ToggleFilterBarVisibility()
 {
-	UNavigationToolSettings* const Settings = GetMutableDefault<UNavigationToolSettings>();
-	if (!ensure(IsValid(Settings)))
+	UNavigationToolSettings* const ToolSettings = GetMutableDefault<UNavigationToolSettings>();
+	if (!ToolSettings)
 	{
 		return;
 	}
 
-	const bool bNewVisibility = !Settings->IsFilterBarVisible();
+	const bool bNewVisibility = !ToolSettings->IsFilterBarVisible();
 
-	Settings->SetFilterBarVisible(bNewVisibility);
+	ToolSettings->SetFilterBarVisible(bNewVisibility);
 
-	StateChangedEvent.Broadcast(bNewVisibility, Settings->GetFilterBarLayout());
+	StateChangedEvent.Broadcast(bNewVisibility, ToolSettings->GetFilterBarLayout());
 }
 
 bool FNavigationToolFilterBar::IsFilterBarLayout(const EFilterBarLayout InLayout) const
 {
-	const UNavigationToolSettings* const Settings = GetDefault<UNavigationToolSettings>();
-	if (ensure(IsValid(Settings)))
+	if (const UNavigationToolSettings* const ToolSettings = GetDefault<UNavigationToolSettings>())
 	{
-		return (Settings->GetFilterBarLayout() == InLayout);
+		return (ToolSettings->GetFilterBarLayout() == InLayout);
 	}
 	return false;
 }
 
 void FNavigationToolFilterBar::SetToVerticalLayout()
 {
-	UNavigationToolSettings* const Settings = GetMutableDefault<UNavigationToolSettings>();
-	if (!ensure(IsValid(Settings)))
+	UNavigationToolSettings* const ToolSettings = GetMutableDefault<UNavigationToolSettings>();
+	if (!ToolSettings)
 	{
 		return;
 	}
 
-	Settings->SetFilterBarLayout(EFilterBarLayout::Vertical);
+	ToolSettings->SetFilterBarLayout(EFilterBarLayout::Vertical);
 
-	StateChangedEvent.Broadcast(IsFilterBarVisible(), Settings->GetFilterBarLayout());
+	StateChangedEvent.Broadcast(IsFilterBarVisible(), ToolSettings->GetFilterBarLayout());
 }
 
 void FNavigationToolFilterBar::SetToHorizontalLayout()
 {
-	UNavigationToolSettings* const Settings = GetMutableDefault<UNavigationToolSettings>();
-	if (!ensure(IsValid(Settings)))
+	UNavigationToolSettings* const ToolSettings = GetMutableDefault<UNavigationToolSettings>();
+	if (!ToolSettings)
 	{
 		return;
 	}
 
-	Settings->SetFilterBarLayout(EFilterBarLayout::Horizontal);
+	ToolSettings->SetFilterBarLayout(EFilterBarLayout::Horizontal);
 
-	StateChangedEvent.Broadcast(IsFilterBarVisible(), Settings->GetFilterBarLayout());
+	StateChangedEvent.Broadcast(IsFilterBarVisible(), ToolSettings->GetFilterBarLayout());
 }
 
 void FNavigationToolFilterBar::ToggleFilterBarLayout()
@@ -1570,8 +1564,8 @@ void FNavigationToolFilterBar::CreateWindow_EditCustomTextFilter(const TSharedPt
 
 bool FNavigationToolFilterBar::CheckFilterNameValidity(const FString& InNewFilterName, const FString& InOldFilterName, const bool bInIsEdit, FText& OutErrorText) const
 {
-	UNavigationToolSettings* const Settings = GetMutableDefault<UNavigationToolSettings>();
-	if (!ensure(IsValid(Settings)))
+	UNavigationToolSettings* const ToolSettings = GetMutableDefault<UNavigationToolSettings>();
+	if (!ToolSettings)
 	{
 		return false;
 	}
@@ -1582,7 +1576,7 @@ bool FNavigationToolFilterBar::CheckFilterNameValidity(const FString& InNewFilte
 		return false;
 	}
 
-	FSequencerFilterBarConfig& Config = Settings->FindOrAddFilterBar(GetIdentifier(), true);
+	FSequencerFilterBarConfig& Config = ToolSettings->FindOrAddFilterBar(GetIdentifier(), true);
 
 	const TArray<FCustomTextFilterData> CustomTextFilterDatas = Config.GetCustomTextFilters();
 
@@ -1604,8 +1598,8 @@ bool FNavigationToolFilterBar::CheckFilterNameValidity(const FString& InNewFilte
 
 bool FNavigationToolFilterBar::TryCreateCustomTextFilter(const FCustomTextFilterData& InNewFilterData, const FString& InOldFilterName, const bool bInApply, FText& OutErrorText)
 {
-	UNavigationToolSettings* const Settings = GetMutableDefault<UNavigationToolSettings>();
-	if (!ensure(IsValid(Settings)))
+	UNavigationToolSettings* const ToolSettings = GetMutableDefault<UNavigationToolSettings>();
+	if (!ToolSettings)
 	{
 		return false;
 	}
@@ -1624,9 +1618,9 @@ bool FNavigationToolFilterBar::TryCreateCustomTextFilter(const FCustomTextFilter
 		return false;
 	}
 
-	FSequencerFilterBarConfig& Config = Settings->FindOrAddFilterBar(GetIdentifier(), false);
+	FSequencerFilterBarConfig& Config = ToolSettings->FindOrAddFilterBar(GetIdentifier(), false);
 	Config.AddCustomTextFilter(InNewFilterData);
-	Settings->SaveConfig();
+	ToolSettings->SaveConfig();
 
 	AddCustomTextFilter(NewCustomTextFilter, false);
 
@@ -1645,15 +1639,15 @@ bool FNavigationToolFilterBar::TryCreateCustomTextFilter(const FCustomTextFilter
 
 bool FNavigationToolFilterBar::TryDeleteCustomTextFilter(const FString& InFilterName, FText& OutErrorText)
 {
-	UNavigationToolSettings* const Settings = GetMutableDefault<UNavigationToolSettings>();
-	if (!ensure(IsValid(Settings)))
+	UNavigationToolSettings* const ToolSettings = GetMutableDefault<UNavigationToolSettings>();
+	if (!ToolSettings)
 	{
 		return false;
 	}
 
-	FSequencerFilterBarConfig& Config = Settings->FindOrAddFilterBar(GetIdentifier(), false);
+	FSequencerFilterBarConfig& Config = ToolSettings->FindOrAddFilterBar(GetIdentifier(), false);
 	Config.RemoveCustomTextFilter(InFilterName);
-	Settings->SaveConfig();
+	ToolSettings->SaveConfig();
 
 	if (const TSharedPtr<FNavigationToolFilter_CustomText> CustomTextFilter = FindCustomTextFilterByDisplayName(InFilterName))
 	{
@@ -1665,8 +1659,8 @@ bool FNavigationToolFilterBar::TryDeleteCustomTextFilter(const FString& InFilter
 
 bool FNavigationToolFilterBar::TryModifyCustomTextFilter(const FCustomTextFilterData& InFilterData, const FString& InOldFilterName, FText& OutErrorText)
 {
-	UNavigationToolSettings* const Settings = GetMutableDefault<UNavigationToolSettings>();
-	if (!ensure(IsValid(Settings)))
+	UNavigationToolSettings* const ToolSettings = GetMutableDefault<UNavigationToolSettings>();
+	if (!ToolSettings)
 	{
 		return false;
 	}
@@ -1688,7 +1682,7 @@ bool FNavigationToolFilterBar::TryModifyCustomTextFilter(const FCustomTextFilter
 	const bool bWasFilterEnabled = IsFilterEnabled(CustomTextFilterRef);
 	const bool bWasFilterActive = IsFilterActive(CustomTextFilterRef);
 
-	FSequencerFilterBarConfig& Config = Settings->FindOrAddFilterBar(GetIdentifier(), false);
+	FSequencerFilterBarConfig& Config = ToolSettings->FindOrAddFilterBar(GetIdentifier(), false);
 	Config.RemoveCustomTextFilter(InOldFilterName);
 
 	if (CustomTextFilter.IsValid())
@@ -1697,7 +1691,7 @@ bool FNavigationToolFilterBar::TryModifyCustomTextFilter(const FCustomTextFilter
 	}
 
 	Config.AddCustomTextFilter(InFilterData);
-	Settings->SaveConfig();
+	ToolSettings->SaveConfig();
 
 	RemoveCustomTextFilter(CustomTextFilterRef, false);
 	AddCustomTextFilter(CustomTextFilterRef, false);
@@ -1745,8 +1739,8 @@ void FNavigationToolFilterBar::OnSearchTextSaved(const FText& InFilterText)
 
 void FNavigationToolFilterBar::BroadcastStateChanged()
 {
-	const UNavigationToolSettings* const Settings = GetDefault<UNavigationToolSettings>();
-	const EFilterBarLayout Layout = IsValid(Settings) ? Settings->GetFilterBarLayout() : EFilterBarLayout::Horizontal;
+	const UNavigationToolSettings* const ToolSettings = GetDefault<UNavigationToolSettings>();
+	const EFilterBarLayout Layout = ToolSettings ? ToolSettings->GetFilterBarLayout() : EFilterBarLayout::Horizontal;
 	StateChangedEvent.Broadcast(IsFilterBarVisible(), Layout);
 }
 
@@ -1764,7 +1758,7 @@ bool FNavigationToolFilterBar::IsFilterSupported(const TSharedRef<FNavigationToo
 	}
 
 	UMovieSceneSequence* const RootMovieSceneSequence = Sequencer->GetRootMovieSceneSequence();
-	if (!IsValid(RootMovieSceneSequence))
+	if (!RootMovieSceneSequence)
 	{
 		return false;
 	}
diff --git a/Engine/Plugins/Experimental/SequenceNavigator/Source/SequenceNavigator/Private/Items/NavigationToolBinding.cpp b/Engine/Plugins/Experimental/SequenceNavigator/Source/SequenceNavigator/Private/Items/NavigationToolBinding.cpp
--- a/Engine/Plugins/Experimental/SequenceNavigator/Source/SequenceNavigator/Private/Items/NavigationToolBinding.cpp
+++ b/Engine/Plugins/Experimental/SequenceNavigator/Source/SequenceNavigator/Private/Items/NavigationToolBinding.cpp
@@ -150,8 +150,7 @@ void FNavigationToolBinding::OnSelect()
 
 void FNavigationToolBinding::OnDoubleClick()
 {
-	UMovieSceneSequence* const Sequence = GetSequence();
-	if (IsValid(Sequence))
+	if (UMovieSceneSequence* const Sequence = GetSequence())
 	{
 		FocusSequence(Tool, *Sequence, *this);
 	}
@@ -288,13 +287,13 @@ const UClass* FNavigationToolBinding::GetBoundObjectClass() const
 	}
 
 	UMovieSceneSequence* const Sequence = GetSequence();
-	if (!IsValid(Sequence))
+	if (!Sequence)
 	{
 		return nullptr;
 	}
 
 	UMovieScene* const MovieScene = Sequence->GetMovieScene();
-	if (!IsValid(MovieScene))
+	if (!MovieScene)
 	{
 		return nullptr;
 	}
@@ -316,7 +315,7 @@ UObject* FNavigationToolBinding::CacheBoundObject()
 	}
 
 	UMovieSceneSequence* const Sequence = GetSequence();
-	if (!IsValid(Sequence))
+	if (!Sequence)
 	{
 		return nullptr;
 	}
diff --git a/Engine/Plugins/Experimental/SequenceNavigator/Source/SequenceNavigator/Private/Items/NavigationToolItemUtils.cpp b/Engine/Plugins/Experimental/SequenceNavigator/Source/SequenceNavigator/Private/Items/NavigationToolItemUtils.cpp
--- a/Engine/Plugins/Experimental/SequenceNavigator/Source/SequenceNavigator/Private/Items/NavigationToolItemUtils.cpp
+++ b/Engine/Plugins/Experimental/SequenceNavigator/Source/SequenceNavigator/Private/Items/NavigationToolItemUtils.cpp
@@ -169,7 +169,7 @@ void AppendSequenceDisplayNameDirtyStatus(FText& InOutDisplayName
 	, const UMovieSceneSequence& InSequence)
 {
 	UPackage* const Package = InSequence.GetPackage();
-	if (IsValid(Package) && Package->IsDirty())
+	if (Package && Package->IsDirty())
 	{
 		InOutDisplayName = FText::Format(LOCTEXT("DirtySymbol", "{0}*"), InOutDisplayName);
 	}
diff --git a/Engine/Plugins/Experimental/SequenceNavigator/Source/SequenceNavigator/Private/Items/NavigationToolMarker.cpp b/Engine/Plugins/Experimental/SequenceNavigator/Source/SequenceNavigator/Private/Items/NavigationToolMarker.cpp
--- a/Engine/Plugins/Experimental/SequenceNavigator/Source/SequenceNavigator/Private/Items/NavigationToolMarker.cpp
+++ b/Engine/Plugins/Experimental/SequenceNavigator/Source/SequenceNavigator/Private/Items/NavigationToolMarker.cpp
@@ -153,7 +153,7 @@ bool FNavigationToolMarker::Delete()
 	}
 
 	UMovieScene* const MovieScene = GetParentMovieScene();
-	if (!IsValid(MovieScene))
+	if (!MovieScene)
 	{
 		return false;
 	}
@@ -180,7 +180,7 @@ bool FNavigationToolMarker::CanRename() const
 	}
 
 	const UMovieScene* const MovieScene = ParentSequenceItem->GetSequenceMovieScene();
-	if (!IsValid(MovieScene))
+	if (!MovieScene)
 	{
 		return false;
 	}
@@ -196,7 +196,7 @@ bool FNavigationToolMarker::Rename(const FString& InName)
 	}
 
 	UMovieScene* const MovieScene = WeakParentSequenceItem.Pin()->GetSequenceMovieScene();
-	if (!IsValid(MovieScene))
+	if (!MovieScene)
 	{
 		return false;
 	}
@@ -227,13 +227,13 @@ FFrameNumber FNavigationToolMarker::GetInTime() const
 	}
 
 	UMovieSceneSequence* const Sequence = WeakParentSequenceItem.Pin()->GetSequence();
-	if (!IsValid(Sequence))
+	if (!Sequence)
 	{
 		return 0;
 	}
 
 	UMovieScene* const MovieScene = Sequence->GetMovieScene();
-	if (!IsValid(MovieScene))
+	if (!MovieScene)
 	{
 		return 0;
 	}
@@ -252,13 +252,13 @@ void FNavigationToolMarker::SetInTime(const FFrameNumber& InTime)
 	}
 
 	UMovieSceneSequence* const Sequence = WeakParentSequenceItem.Pin()->GetSequence();
-	if (!IsValid(Sequence))
+	if (!Sequence)
 	{
 		return;
 	}
 
 	UMovieScene* const MovieScene = Sequence->GetMovieScene();
-	if (!IsValid(MovieScene))
+	if (!MovieScene)
 	{
 		return;
 	}
@@ -293,7 +293,7 @@ FMovieSceneMarkedFrame* FNavigationToolMarker::GetMarkedFrame() const
 	}
 
 	UMovieScene* const MovieScene = GetParentMovieScene();
-	if (!IsValid(MovieScene))
+	if (!MovieScene)
 	{
 		return nullptr;
 	}
diff --git a/Engine/Plugins/Experimental/SequenceNavigator/Source/SequenceNavigator/Private/Items/NavigationToolSubTrack.cpp b/Engine/Plugins/Experimental/SequenceNavigator/Source/SequenceNavigator/Private/Items/NavigationToolSubTrack.cpp
--- a/Engine/Plugins/Experimental/SequenceNavigator/Source/SequenceNavigator/Private/Items/NavigationToolSubTrack.cpp
+++ b/Engine/Plugins/Experimental/SequenceNavigator/Source/SequenceNavigator/Private/Items/NavigationToolSubTrack.cpp
@@ -41,15 +41,17 @@ void FNavigationToolSubTrack::FindChildren(TArray<FNavigationToolItemPtr>& OutCh
 	if (UMovieSceneSubTrack* const SubTrack = GetSubTrack())
 	{
 		const TSharedRef<FNavigationToolSubTrack> SharedThisRef = SharedThis(this);
-
+		const TSharedPtr<FNavigationToolProvider> Provider = GetProvider();
 		const TArray<UMovieSceneSection*>& AllSections = SubTrack->GetAllSections();
+
 		for (int32 Index = 0; Index < AllSections.Num(); ++Index)
 		{
 			if (UMovieSceneSubSection* const SubSection = Cast<UMovieSceneSubSection>(AllSections[Index]))
 			{
 				if (UMovieSceneSequence* const Sequence = SubSection->GetSequence())
 				{
-					const FNavigationToolItemPtr NewItem = Tool.FindOrAdd<FNavigationToolSequence>(GetProvider(), SharedThisRef, Sequence, SubSection, Index);
+					const FNavigationToolItemPtr NewItem = Tool.FindOrAdd<FNavigationToolSequence>(Provider
+						, SharedThisRef, Sequence, SubSection, Index);
 					OutChildren.Add(NewItem);
 					if (bInRecursive)
 					{
diff --git a/Engine/Plugins/Experimental/SequenceNavigator/Source/SequenceNavigator/Private/Items/NavigationToolTrack.cpp b/Engine/Plugins/Experimental/SequenceNavigator/Source/SequenceNavigator/Private/Items/NavigationToolTrack.cpp
--- a/Engine/Plugins/Experimental/SequenceNavigator/Source/SequenceNavigator/Private/Items/NavigationToolTrack.cpp
+++ b/Engine/Plugins/Experimental/SequenceNavigator/Source/SequenceNavigator/Private/Items/NavigationToolTrack.cpp
@@ -237,7 +237,7 @@ EItemContainsPlayhead FNavigationToolTrack::ContainsPlayhead() const
 void FNavigationToolTrack::OnTrackObjectChanged()
 {
 	UMovieSceneTrack* const Track = GetTrack();
-	if (!IsValid(Track))
+	if (!Track)
 	{
 		return;
 	}
@@ -284,7 +284,7 @@ UE::Sequencer::TViewModelPtr<UE::Sequencer::FTrackModel> FNavigationToolTrack::G
 	using namespace UE::Sequencer;
 
 	UMovieSceneTrack* const Track = GetTrack();
-	if (!IsValid(Track))
+	if (!Track)
 	{
 		return nullptr;
 	}
diff --git a/Engine/Plugins/Experimental/SequenceNavigator/Source/SequenceNavigator/Private/Menus/NavigationToolFilterContextMenu.cpp b/Engine/Plugins/Experimental/SequenceNavigator/Source/SequenceNavigator/Private/Menus/NavigationToolFilterContextMenu.cpp
--- a/Engine/Plugins/Experimental/SequenceNavigator/Source/SequenceNavigator/Private/Menus/NavigationToolFilterContextMenu.cpp
+++ b/Engine/Plugins/Experimental/SequenceNavigator/Source/SequenceNavigator/Private/Menus/NavigationToolFilterContextMenu.cpp
@@ -222,17 +222,16 @@ void FNavigationToolFilterContextMenu::OnDeleteFilter()
 	FilterInterface.SetFilterActiveByDisplayName(FilterDisplayName, false);
 	FilterInterface.RemoveCustomTextFilter(CustomTextFilter.ToSharedRef(), false);
 
-	UNavigationToolSettings* const Settings = GetMutableDefault<UNavigationToolSettings>();
-	if (IsValid(Settings))
+	if (UNavigationToolSettings* const ToolSettings = GetMutableDefault<UNavigationToolSettings>())
 	{
-		FSequencerFilterBarConfig& Config = Settings->FindOrAddFilterBar(FilterInterface.GetIdentifier(), false);
+		FSequencerFilterBarConfig& Config = ToolSettings->FindOrAddFilterBar(FilterInterface.GetIdentifier(), false);
 
 		if (!Config.RemoveCustomTextFilter(FilterDisplayName))
 		{
 			return;
 		}
 
-		Settings->SaveConfig();
+		ToolSettings->SaveConfig();
 	}
 }
 
diff --git a/Engine/Plugins/Experimental/SequenceNavigator/Source/SequenceNavigator/Private/Menus/NavigationToolFilterMenu.cpp b/Engine/Plugins/Experimental/SequenceNavigator/Source/SequenceNavigator/Private/Menus/NavigationToolFilterMenu.cpp
--- a/Engine/Plugins/Experimental/SequenceNavigator/Source/SequenceNavigator/Private/Menus/NavigationToolFilterMenu.cpp
+++ b/Engine/Plugins/Experimental/SequenceNavigator/Source/SequenceNavigator/Private/Menus/NavigationToolFilterMenu.cpp
@@ -99,8 +99,8 @@ void FNavigationToolFilterMenu::PopulateFilterOptionsSection(UToolMenu& InMenu)
 			FCanExecuteAction::CreateSP(this, &FNavigationToolFilterMenu::CanResetFilters)
 		));
 
-	const UNavigationToolSettings* const Settings = GetDefault<UNavigationToolSettings>();
-	if (Settings->ShouldUseFilterSubmenusForCategories())
+	const UNavigationToolSettings* const ToolSettings = GetDefault<UNavigationToolSettings>();
+	if (ToolSettings && ToolSettings->ShouldUseFilterSubmenusForCategories())
 	{
 		Section.AddSeparator(NAME_None);
 	}
@@ -114,13 +114,13 @@ void FNavigationToolFilterMenu::PopulateCommonFilterSections(UToolMenu& InMenu)
 		return;
 	}
 
-	const UNavigationToolSettings* const Settings = GetDefault<UNavigationToolSettings>();
-	if (!Settings)
+	const UNavigationToolSettings* const ToolSettings = GetDefault<UNavigationToolSettings>();
+	if (!ToolSettings)
 	{
 		return;
 	}
 
-	if (Settings->ShouldUseFilterSubmenusForCategories())
+	if (ToolSettings->ShouldUseFilterSubmenusForCategories())
 	{
 		FToolMenuSection& Section = InMenu.FindOrAddSection(TEXT("OtherFilters"), LOCTEXT("OtherFiltersHeading", "Other Filters"));
 
@@ -164,6 +164,12 @@ void FNavigationToolFilterMenu::FillCustomTextFiltersMenu(UToolMenu* const InMen
 		return;
 	}
 
+	UNavigationToolSettings* const ToolSettings = GetMutableDefault<UNavigationToolSettings>();
+	if (!ToolSettings)
+	{
+		return;
+	}
+
 	const TSharedRef<INavigationToolFilterBar> FilterBarRef = FilterBar.ToSharedRef();
 
 	FToolMenuSection& CustomTextOptionsSection = InMenu->FindOrAddSection(TEXT("CustomTextFilterOptions")
@@ -188,10 +194,7 @@ void FNavigationToolFilterMenu::FillCustomTextFiltersMenu(UToolMenu* const InMen
 		FSlateIcon(FAppStyle::Get().GetStyleSetName(), TEXT("Icons.PlusCircle")),
 		FExecuteAction::CreateSP(FilterBarRef, &INavigationToolFilterBar::CreateNewTextFilter));
 
-	UNavigationToolSettings* const Settings = GetMutableDefault<UNavigationToolSettings>();
-	check(Settings);
-
-	FSequencerFilterBarConfig& Config = Settings->FindOrAddFilterBar(FilterBar->GetIdentifier(), true);
+	FSequencerFilterBarConfig& Config = ToolSettings->FindOrAddFilterBar(FilterBar->GetIdentifier(), true);
 
 	const TArray<FCustomTextFilterData>& CustomTextFilterDatas = Config.GetCustomTextFilters();
 	if (CustomTextFilterDatas.IsEmpty())
@@ -221,6 +224,10 @@ void FNavigationToolFilterMenu::FillFiltersMenuCategory(FToolMenuSection& InOutS
 	}
 
 	UMovieSceneSequence* const FocusedSequence = FilterBar->GetSequencer().GetFocusedMovieSceneSequence();
+	if (!FocusedSequence)
+	{
+		return;
+	}
 
 	for (const TSharedRef<FNavigationToolFilter>& Filter : FilterBar->GetCommonFilters())
 	{
@@ -262,7 +269,7 @@ void FNavigationToolFilterMenu::FillFiltersMenuCategory(FToolMenuSection& InOutS
 
 void FNavigationToolFilterMenu::FillFiltersMenuCategory(UToolMenu* const InMenu, const TSharedRef<FFilterCategory> InMenuCategory)
 {
-	if (IsValid(InMenu))
+	if (InMenu)
 	{
 		FillFiltersMenuCategory(InMenu->AddSection(TEXT("Section")), InMenuCategory);
 	}
diff --git a/Engine/Plugins/Experimental/SequenceNavigator/Source/SequenceNavigator/Private/Menus/NavigationToolToolbarMenu.cpp b/Engine/Plugins/Experimental/SequenceNavigator/Source/SequenceNavigator/Private/Menus/NavigationToolToolbarMenu.cpp
--- a/Engine/Plugins/Experimental/SequenceNavigator/Source/SequenceNavigator/Private/Menus/NavigationToolToolbarMenu.cpp
+++ b/Engine/Plugins/Experimental/SequenceNavigator/Source/SequenceNavigator/Private/Menus/NavigationToolToolbarMenu.cpp
@@ -282,13 +282,13 @@ void FNavigationToolToolbarMenu::CreateItemViewOptionsMenu(UToolMenu* const InMe
 
 void FNavigationToolToolbarMenu::CreateColumnViewOptionsMenu(UToolMenu* const InMenu)
 {
-	if (!IsValid(InMenu))
+	if (!InMenu)
 	{
 		return;
 	}
 
 	UNavigationToolViewMenuContext* const ToolBarContext = InMenu->FindContext<UNavigationToolViewMenuContext>();
-	if (!IsValid(ToolBarContext))
+	if (!ToolBarContext)
 	{
 		return;
 	}
@@ -302,7 +302,7 @@ void FNavigationToolToolbarMenu::CreateColumnViewOptionsMenu(UToolMenu* const In
 	const TSharedRef<FNavigationToolView> ToolViewRef = ToolView.ToSharedRef();
 
 	UNavigationToolSettings* const ToolSettings = GetMutableDefault<UNavigationToolSettings>();
-	if (!ensure(IsValid(ToolSettings)))
+	if (!ToolSettings)
 	{
 		return;
 	}
@@ -340,7 +340,7 @@ void FNavigationToolToolbarMenu::CreateColumnViewOptionsMenu(UToolMenu* const In
 
 void FNavigationToolToolbarMenu::CreateFilterBarOptionsMenu(UToolMenu* const InMenu)
 {
-	if (!IsValid(InMenu))
+	if (!InMenu)
 	{
 		return;
 	}
@@ -362,7 +362,7 @@ void FNavigationToolToolbarMenu::CreateFilterBarOptionsMenu(UToolMenu* const InM
 ECheckBoxState FNavigationToolToolbarMenu::GetCustomColumnViewMenuItemCheckState(const FText InColumnViewName, const TSharedRef<FNavigationToolView> InToolView) const
 {
 	UNavigationToolSettings* const ToolSettings = GetMutableDefault<UNavigationToolSettings>();
-	if (!ensure(IsValid(ToolSettings)))
+	if (!ToolSettings)
 	{
 		return ECheckBoxState::Unchecked;
 	}
@@ -389,7 +389,7 @@ ECheckBoxState FNavigationToolToolbarMenu::GetCustomColumnViewMenuItemCheckState
 void FNavigationToolToolbarMenu::OnDeleteCustomColumnViewMenuItemClick(const FText InColumnViewName)
 {
 	UNavigationToolSettings* const ToolSettings = GetMutableDefault<UNavigationToolSettings>();
-	if (!ensure(IsValid(ToolSettings)))
+	if (!ToolSettings)
 	{
 		return;
 	}
@@ -415,8 +415,7 @@ void FNavigationToolToolbarMenu::OnToggleGlobalFilter(const TSharedRef<FNavigati
 
 	InFilter->SetActive(bNewActiveState);
 
-	UNavigationToolSettings* const ToolSettings = GetMutableDefault<UNavigationToolSettings>();
-	if (ensure(IsValid(ToolSettings)))
+	if (UNavigationToolSettings* const ToolSettings = GetMutableDefault<UNavigationToolSettings>())
 	{
 		ToolSettings->SetBuiltInFilterEnabled(InFilter->GetFilterParams().GetFilterId(), bNewActiveState);
 	}
diff --git a/Engine/Plugins/Experimental/SequenceNavigator/Source/SequenceNavigator/Private/NavigationTool.cpp b/Engine/Plugins/Experimental/SequenceNavigator/Source/SequenceNavigator/Private/NavigationTool.cpp
--- a/Engine/Plugins/Experimental/SequenceNavigator/Source/SequenceNavigator/Private/NavigationTool.cpp
+++ b/Engine/Plugins/Experimental/SequenceNavigator/Source/SequenceNavigator/Private/NavigationTool.cpp
@@ -521,8 +521,7 @@ TArray<FNavigationToolItemPtr> FNavigationTool::TryFindItems(const UE::Sequencer
 
 	if (const TViewModelPtr<FSectionModel> SectionModel = InViewModel.ImplicitCast())
 	{
-		UMovieSceneSubSection* const SubSection = Cast<UMovieSceneSubSection>(SectionModel->GetSection());
-		if (IsValid(SubSection))
+		if (UMovieSceneSubSection* const SubSection = Cast<UMovieSceneSubSection>(SectionModel->GetSection()))
 		{
 			return FindItemsFromMovieSceneObject(SubSection->GetSequence());
 		}
@@ -1270,7 +1269,7 @@ void FNavigationTool::OnSequencerSelectionChanged()
 TArray<FNavigationToolItemPtr> FNavigationTool::FindItemsFromMovieSceneObject(UObject* const InObject) const
 {
 	TArray<FNavigationToolItemPtr> OutItems;
-	if (IsValid(InObject))
+	if (InObject)
 	{
 		for (const FNavigationToolItemPtr& Item : RootItem->GetDescendants())
 		{
diff --git a/Engine/Plugins/Experimental/SequenceNavigator/Source/SequenceNavigator/Private/NavigationToolTab.cpp b/Engine/Plugins/Experimental/SequenceNavigator/Source/SequenceNavigator/Private/NavigationToolTab.cpp
--- a/Engine/Plugins/Experimental/SequenceNavigator/Source/SequenceNavigator/Private/NavigationToolTab.cpp
+++ b/Engine/Plugins/Experimental/SequenceNavigator/Source/SequenceNavigator/Private/NavigationToolTab.cpp
@@ -176,7 +176,7 @@ TSharedRef<SDockTab> FNavigationToolTab::SpawnToolTab(const FSpawnTabArgs& InArg
 		.OnTabClosed_Lambda([this](const TSharedRef<SDockTab> InDockTab)
 			{
 				USequencerSettings* const Settings = GetSequencerSettings();
-				if (IsValid(Settings) && !bShuttingDown)
+				if (Settings && !bShuttingDown)
 				{
 					Settings->SetNavigationToolVisible(false);
 				}
@@ -187,8 +187,7 @@ TSharedRef<SDockTab> FNavigationToolTab::SpawnToolTab(const FSpawnTabArgs& InArg
 			ToolWidget.ToSharedRef()
 		];
 
-	USequencerSettings* const SequencerSettings = GetSequencerSettings();
-	if (IsValid(SequencerSettings))
+	if (USequencerSettings* const SequencerSettings = GetSequencerSettings())
 	{
 		SequencerSettings->SetNavigationToolVisible(true);
 	}
diff --git a/Engine/Plugins/Experimental/SequenceNavigator/Source/SequenceNavigator/Private/NavigationToolView.cpp b/Engine/Plugins/Experimental/SequenceNavigator/Source/SequenceNavigator/Private/NavigationToolView.cpp
--- a/Engine/Plugins/Experimental/SequenceNavigator/Source/SequenceNavigator/Private/NavigationToolView.cpp
+++ b/Engine/Plugins/Experimental/SequenceNavigator/Source/SequenceNavigator/Private/NavigationToolView.cpp
@@ -537,7 +537,7 @@ void FNavigationToolView::GetChildrenOfItem(const FNavigationToolItemPtr& InItem
 		const bool bShouldShowItem = ShouldShowItem(ChildItem, true, InViewMode);
 		if (bShouldShowItem)
 		{
-			// If current item is visible in outliner, add it to the children
+			// If the current item is visible in outliner, add it to the children
 			OutChildren.Add(ChildItem);
 		}
 		else if (!InRecursionDisallowedItems.Contains(ChildItem))
@@ -550,7 +550,7 @@ void FNavigationToolView::GetChildrenOfItem(const FNavigationToolItemPtr& InItem
 			const bool bShouldShowItemWithoutFilters = ShouldShowItem(ChildItem, false, InViewMode);
 			const bool bShouldMuteItem = bShouldUseMutedHierarchy && bShouldShowItemWithoutFilters;
 
-			// If it's muted hierarchy, there might be ONLY grand children that are just visible in other view modes, 
+			// If Muted Hierarchy, there might be ONLY grand children that are just visible in other view modes, 
 			// so instead of just filtering out the child item, check that there are no grand children from other view modes passing filter tests
 			// If it's NOT muted hierarchy, just get the grand children visible in the requested view mode, as this ChildItem is guaranteed to be hidden
 			const ENavigationToolItemViewMode ViewModeToUse = bShouldMuteItem ? ENavigationToolItemViewMode::All : InViewMode;
@@ -2103,7 +2103,7 @@ bool FNavigationToolView::CanFocusSingleSelection() const
 	if (SelectedItems.Num() == 1)
 	{
 		const FNavigationToolSequence* const SequenceItem = SelectedItems[0]->CastTo<FNavigationToolSequence>();
-		return SequenceItem && IsValid(SequenceItem->GetSequence());
+		return SequenceItem && SequenceItem->GetSequence();
 	}
 	return false;
 }
@@ -2122,7 +2122,7 @@ void FNavigationToolView::FocusSingleSelection()
 	}
 
 	UMovieSceneSequence* const Sequence = SequenceItem->GetSequence();
-	if (!IsValid(Sequence))
+	if (!Sequence)
 	{
 		return;
 	}
@@ -2150,7 +2150,7 @@ bool FNavigationToolView::CanFocusInContentBrowser() const
 	}
 
 	UMovieSceneSequence* const Sequence = SequenceItem->GetSequence();
-	if (!IsValid(Sequence))
+	if (!Sequence)
 	{
 		return false;
 	}
@@ -2172,13 +2172,13 @@ void FNavigationToolView::FocusInContentBrowser()
 	}
 
 	UMovieSceneSequence* const Sequence = SequenceItem->GetSequence();
-	if (!IsValid(Sequence))
+	if (!Sequence)
 	{
 		return;
 	}
 
 	ULevelSequence* const LevelSequence = Cast<ULevelSequence>(Sequence);
-	if (!IsValid(LevelSequence))
+	if (!LevelSequence)
 	{
 		return;
 	}
diff --git a/Engine/Plugins/Experimental/SequenceNavigator/Source/SequenceNavigator/Private/Providers/NavigationToolProvider.cpp b/Engine/Plugins/Experimental/SequenceNavigator/Source/SequenceNavigator/Private/Providers/NavigationToolProvider.cpp
--- a/Engine/Plugins/Experimental/SequenceNavigator/Source/SequenceNavigator/Private/Providers/NavigationToolProvider.cpp
+++ b/Engine/Plugins/Experimental/SequenceNavigator/Source/SequenceNavigator/Private/Providers/NavigationToolProvider.cpp
@@ -166,20 +166,20 @@ void FNavigationToolProvider::Deactivate(FNavigationTool& InTool)
 
 void FNavigationToolProvider::CleanupExtendedColumnViews()
 {
-	UNavigationToolSettings* const Settings = GetMutableDefault<UNavigationToolSettings>();
-	if (!Settings)
+	UNavigationToolSettings* const ToolSettings = GetMutableDefault<UNavigationToolSettings>();
+	if (!ToolSettings)
 	{
 		return;
 	}
 
-	TSet<FNavigationToolColumnView>& CustomColumnViews = Settings->GetCustomColumnViews();
+	TSet<FNavigationToolColumnView>& CustomColumnViews = ToolSettings->GetCustomColumnViews();
 
 	for (const FText& ColumnViewName : ExtendedColumnViewNames)
 	{
 		CustomColumnViews.Remove(ColumnViewName);
 	}
 
-	Settings->SaveConfig();
+	ToolSettings->SaveConfig();
 }
 
 bool FNavigationToolProvider::IsSequenceSupported(UMovieSceneSequence* const InSequence) const
diff --git a/Engine/Plugins/Experimental/SequenceNavigator/Source/SequenceNavigator/Private/Utils/NavigationToolMiscUtils.cpp b/Engine/Plugins/Experimental/SequenceNavigator/Source/SequenceNavigator/Private/Utils/NavigationToolMiscUtils.cpp
--- a/Engine/Plugins/Experimental/SequenceNavigator/Source/SequenceNavigator/Private/Utils/NavigationToolMiscUtils.cpp
+++ b/Engine/Plugins/Experimental/SequenceNavigator/Source/SequenceNavigator/Private/Utils/NavigationToolMiscUtils.cpp
@@ -218,24 +218,21 @@ void FocusItemInSequencer(const INavigationTool& InTool, const FNavigationToolIt
 {
 	if (const FNavigationToolSequence* const SequenceItem = InItem->CastTo<FNavigationToolSequence>())
 	{
-		UMovieSceneSequence* const ParentSequence = SequenceItem->GetSequence();
-		if (IsValid(ParentSequence))
+		if (UMovieSceneSequence* const ParentSequence = SequenceItem->GetSequence())
 		{
 			FocusSequence(InTool, *ParentSequence);
 		}
 	}
 	else if (const FNavigationToolTrack* const TrackItem = InItem->CastTo<FNavigationToolTrack>())
 	{
-		UMovieSceneSequence* const Sequence = TrackItem->GetSequence();
-		if (IsValid(Sequence))
+		if (UMovieSceneSequence* const Sequence = TrackItem->GetSequence())
 		{
 			FocusSequence(InTool, *Sequence, *TrackItem);
 		}
 	}
 	else if (const FNavigationToolBinding* const BindingItem = InItem->CastTo<FNavigationToolBinding>())
 	{
-		UMovieSceneSequence* const Sequence = BindingItem->GetSequence();
-		if (IsValid(Sequence))
+		if (UMovieSceneSequence* const Sequence = BindingItem->GetSequence())
 		{
 			FocusSequence(InTool, *Sequence, *BindingItem);
 		}
diff --git a/Engine/Plugins/Experimental/SequenceNavigator/Source/SequenceNavigator/Private/Widgets/Columns/SNavigationToolComment.cpp b/Engine/Plugins/Experimental/SequenceNavigator/Source/SequenceNavigator/Private/Widgets/Columns/SNavigationToolComment.cpp
--- a/Engine/Plugins/Experimental/SequenceNavigator/Source/SequenceNavigator/Private/Widgets/Columns/SNavigationToolComment.cpp
+++ b/Engine/Plugins/Experimental/SequenceNavigator/Source/SequenceNavigator/Private/Widgets/Columns/SNavigationToolComment.cpp
@@ -25,7 +25,7 @@ void SNavigationToolComment::Construct(const FArguments& InArgs
 	WeakRowWidget = InRowWidget;
 
 	const FNavigationToolSequence* const SequenceItem = InItem->CastTo<FNavigationToolSequence>();
-	if (!SequenceItem || !IsValid(SequenceItem->GetSubSection()))
+	if (!SequenceItem || !SequenceItem->GetSubSection())
 	{
 		return;
 	}
@@ -45,8 +45,7 @@ SNavigationToolComment::~SNavigationToolComment()
 
 FString SNavigationToolComment::GetMetaDataComment() const
 {
-	UMovieSceneMetaData* const SequenceMetaData = GetSequenceItemMetaData(WeakItem.Pin());
-	if (IsValid(SequenceMetaData))
+	if (UMovieSceneMetaData* const SequenceMetaData = GetSequenceItemMetaData(WeakItem.Pin()))
 	{
 		return SequenceMetaData->GetNotes();
 	}
@@ -72,13 +71,13 @@ void SNavigationToolComment::OnCommentTextCommitted(const FText& InNewText, cons
 	}
 
 	UMovieSceneSequence* const Sequence = SequenceItem->GetSequence();
-	if (!IsValid(Sequence))
+	if (!Sequence)
 	{
 		return;
 	}
 
 	UMovieSceneMetaData* const SequenceMetaData = FSubTrackEditorUtil::FindOrAddMetaData(Sequence);
-	if (!IsValid(SequenceMetaData))
+	if (!SequenceMetaData)
 	{
 		return;
 	}
diff --git a/Engine/Plugins/Experimental/SequenceNavigator/Source/SequenceNavigator/Private/Widgets/Columns/SNavigationToolEndFrameOffset.cpp b/Engine/Plugins/Experimental/SequenceNavigator/Source/SequenceNavigator/Private/Widgets/Columns/SNavigationToolEndFrameOffset.cpp
--- a/Engine/Plugins/Experimental/SequenceNavigator/Source/SequenceNavigator/Private/Widgets/Columns/SNavigationToolEndFrameOffset.cpp
+++ b/Engine/Plugins/Experimental/SequenceNavigator/Source/SequenceNavigator/Private/Widgets/Columns/SNavigationToolEndFrameOffset.cpp
@@ -22,7 +22,7 @@ void SNavigationToolEndFrameOffset::Construct(const FArguments& InArgs
 	WeakRowWidget = InRowWidget;
 
 	const FNavigationToolSequence* const SequenceItem = InItem->CastTo<FNavigationToolSequence>();
-	if (!SequenceItem || !IsValid(SequenceItem->GetSubSection()))
+	if (!SequenceItem || !SequenceItem->GetSubSection())
 	{
 		return;
 	}
@@ -32,8 +32,7 @@ void SNavigationToolEndFrameOffset::Construct(const FArguments& InArgs
 
 double SNavigationToolEndFrameOffset::GetFrameTimeValue() const
 {
-	UMovieSceneSubSection* const SubSection = Cast<UMovieSceneSubSection>(GetSequenceItemSubSection(WeakItem.Pin()));
-	if (IsValid(SubSection))
+	if (UMovieSceneSubSection* const SubSection = Cast<UMovieSceneSubSection>(GetSequenceItemSubSection(WeakItem.Pin())))
 	{
 		return SubSection->Parameters.EndFrameOffset.Value;
 	}
@@ -43,7 +42,7 @@ double SNavigationToolEndFrameOffset::GetFrameTimeValue() const
 void SNavigationToolEndFrameOffset::OnFrameTimeValueCommitted(const double InNewValue, const ETextCommit::Type InCommitType)
 {
 	UMovieSceneSubSection* const SubSection = Cast<UMovieSceneSubSection>(GetSequenceItemSubSection(WeakItem.Pin()));
-	if (!IsValid(SubSection))
+	if (!SubSection)
 	{
 		return;
 	}
diff --git a/Engine/Plugins/Experimental/SequenceNavigator/Source/SequenceNavigator/Private/Widgets/Columns/SNavigationToolHBias.cpp b/Engine/Plugins/Experimental/SequenceNavigator/Source/SequenceNavigator/Private/Widgets/Columns/SNavigationToolHBias.cpp
--- a/Engine/Plugins/Experimental/SequenceNavigator/Source/SequenceNavigator/Private/Widgets/Columns/SNavigationToolHBias.cpp
+++ b/Engine/Plugins/Experimental/SequenceNavigator/Source/SequenceNavigator/Private/Widgets/Columns/SNavigationToolHBias.cpp
@@ -24,7 +24,7 @@ void SNavigationToolHBias::Construct(const FArguments& InArgs
 	WeakRowWidget = InRowWidget;
 
 	const FNavigationToolSequence* const SequenceItem = InItem->CastTo<FNavigationToolSequence>();
-	if (!SequenceItem || !IsValid(SequenceItem->GetSubSection()))
+	if (!SequenceItem || !SequenceItem->GetSubSection())
 	{
 		return;
 	}
@@ -45,7 +45,7 @@ void SNavigationToolHBias::Construct(const FArguments& InArgs
 int32 SNavigationToolHBias::GetValue() const
 {
 	UMovieSceneSubSection* const SubSection = Cast<UMovieSceneSubSection>(GetSequenceItemSubSection(WeakItem.Pin()));
-	if (!IsValid(SubSection))
+	if (!SubSection)
 	{
 		return 0;
 	}
@@ -61,7 +61,7 @@ void SNavigationToolHBias::OnValueChanged(const int32 InNewValue)
 void SNavigationToolHBias::OnValueCommitted(const int32 InNewValue, const ETextCommit::Type InCommitType)
 {
 	UMovieSceneSubSection* const SubSection = Cast<UMovieSceneSubSection>(GetSequenceItemSubSection(WeakItem.Pin()));
-	if (!IsValid(SubSection))
+	if (!SubSection)
 	{
 		return;
 	}
diff --git a/Engine/Plugins/Experimental/SequenceNavigator/Source/SequenceNavigator/Private/Widgets/Columns/SNavigationToolLength.cpp b/Engine/Plugins/Experimental/SequenceNavigator/Source/SequenceNavigator/Private/Widgets/Columns/SNavigationToolLength.cpp
--- a/Engine/Plugins/Experimental/SequenceNavigator/Source/SequenceNavigator/Private/Widgets/Columns/SNavigationToolLength.cpp
+++ b/Engine/Plugins/Experimental/SequenceNavigator/Source/SequenceNavigator/Private/Widgets/Columns/SNavigationToolLength.cpp
@@ -23,7 +23,7 @@ void SNavigationToolLength::Construct(const FArguments& InArgs
 	WeakRowWidget = InRowWidget;
 
 	const FNavigationToolSequence* const SequenceItem = InItem->CastTo<FNavigationToolSequence>();
-	if (!SequenceItem || !IsValid(SequenceItem->GetSubSection()))
+	if (!SequenceItem || !SequenceItem->GetSubSection())
 	{
 		return;
 	}
@@ -33,8 +33,7 @@ void SNavigationToolLength::Construct(const FArguments& InArgs
 
 double SNavigationToolLength::GetFrameTimeValue() const
 {
-	UMovieSceneSubSection* const SubSection = Cast<UMovieSceneSubSection>(GetSequenceItemSubSection(WeakItem.Pin()));
-	if (IsValid(SubSection))
+	if (UMovieSceneSubSection* const SubSection = Cast<UMovieSceneSubSection>(GetSequenceItemSubSection(WeakItem.Pin())))
 	{
 		return SubSection->GetRange().Size<FFrameNumber>().Value;
 	}
diff --git a/Engine/Plugins/Experimental/SequenceNavigator/Source/SequenceNavigator/Private/Widgets/Columns/SNavigationToolStartFrameOffset.cpp b/Engine/Plugins/Experimental/SequenceNavigator/Source/SequenceNavigator/Private/Widgets/Columns/SNavigationToolStartFrameOffset.cpp
--- a/Engine/Plugins/Experimental/SequenceNavigator/Source/SequenceNavigator/Private/Widgets/Columns/SNavigationToolStartFrameOffset.cpp
+++ b/Engine/Plugins/Experimental/SequenceNavigator/Source/SequenceNavigator/Private/Widgets/Columns/SNavigationToolStartFrameOffset.cpp
@@ -21,7 +21,7 @@ void SNavigationToolStartFrameOffset::Construct(const FArguments& InArgs
 	WeakRowWidget = InRowWidget;
 
 	const FNavigationToolSequence* const SequenceItem = InItem->CastTo<FNavigationToolSequence>();
-	if (!SequenceItem || !IsValid(SequenceItem->GetSubSection()))
+	if (!SequenceItem || !SequenceItem->GetSubSection())
 	{
 		return;
 	}
@@ -31,8 +31,7 @@ void SNavigationToolStartFrameOffset::Construct(const FArguments& InArgs
 
 double SNavigationToolStartFrameOffset::GetFrameTimeValue() const
 {
-	const UMovieSceneSubSection* const SubSection = Cast<UMovieSceneSubSection>(GetSequenceItemSubSection(WeakItem.Pin()));
-	if (IsValid(SubSection))
+	if (const UMovieSceneSubSection* const SubSection = Cast<UMovieSceneSubSection>(GetSequenceItemSubSection(WeakItem.Pin())))
 	{
 		return SubSection->Parameters.StartFrameOffset.Value;
 	}
@@ -42,7 +41,7 @@ double SNavigationToolStartFrameOffset::GetFrameTimeValue() const
 void SNavigationToolStartFrameOffset::OnFrameTimeValueCommitted(const double InNewValue, const ETextCommit::Type InCommitType)
 {
 	UMovieSceneSubSection* const SubSection = Cast<UMovieSceneSubSection>(GetSequenceItemSubSection(WeakItem.Pin()));
-	if (!IsValid(SubSection))
+	if (!SubSection)
 	{
 		return;
 	}
diff --git a/Engine/Plugins/Experimental/SequenceNavigator/Source/SequenceNavigator/Private/Widgets/Columns/SNavigationToolTake.cpp b/Engine/Plugins/Experimental/SequenceNavigator/Source/SequenceNavigator/Private/Widgets/Columns/SNavigationToolTake.cpp
--- a/Engine/Plugins/Experimental/SequenceNavigator/Source/SequenceNavigator/Private/Widgets/Columns/SNavigationToolTake.cpp
+++ b/Engine/Plugins/Experimental/SequenceNavigator/Source/SequenceNavigator/Private/Widgets/Columns/SNavigationToolTake.cpp
@@ -46,7 +46,7 @@ public:
 				[
 					SNew(STextBlock)
 					.Text(FText::FromString(TakeEntry->DisplayName))
-					.ColorAndOpacity(FLinearColor::White)
+					.ColorAndOpacity(FStyleColors::Foreground)
 				]
 
 			]
@@ -267,7 +267,7 @@ void SNavigationToolTake::CacheTakes()
 	CachedTakes.Reset();
 
 	UMovieSceneSubSection* const SubSection = Cast<UMovieSceneSubSection>(GetSequenceItemSubSection(WeakItem.Pin()));
-	if (!IsValid(SubSection))
+	if (!SubSection)
 	{
 		return;
 	}
@@ -294,8 +294,7 @@ void SNavigationToolTake::CacheTakes()
 		uint32 TakeNumber = INDEX_NONE;
 		if (MovieSceneToolHelpers::GetTakeNumber(SubSection, ThisAssetData, TakeNumber))
 		{
-			UMovieSceneSequence* const Sequence = Cast<UMovieSceneSequence>(ThisAssetData.GetAsset());
-			if (IsValid(Sequence))
+			if (UMovieSceneSequence* const Sequence = Cast<UMovieSceneSequence>(ThisAssetData.GetAsset()))
 			{
 				const TSharedRef<FTakeItemInfo> NewTakeInfo = MakeShared<FTakeItemInfo>();
 				NewTakeInfo->TakeIndex = TakeIndex;
diff --git a/Engine/Plugins/Experimental/SequenceNavigator/Source/SequenceNavigator/Private/Widgets/SNavigationToolFilterBar.cpp b/Engine/Plugins/Experimental/SequenceNavigator/Source/SequenceNavigator/Private/Widgets/SNavigationToolFilterBar.cpp
--- a/Engine/Plugins/Experimental/SequenceNavigator/Source/SequenceNavigator/Private/Widgets/SNavigationToolFilterBar.cpp
+++ b/Engine/Plugins/Experimental/SequenceNavigator/Source/SequenceNavigator/Private/Widgets/SNavigationToolFilterBar.cpp
@@ -407,13 +407,13 @@ void SNavigationToolFilterBar::OnEnableAllGroupFilters(bool bEnableAll)
 	}
 
 	UMovieSceneSequence* const FocusedMovieSequence = FilterBar->GetSequencer().GetFocusedMovieSceneSequence();
-	if (!IsValid(FocusedMovieSequence))
+	if (!FocusedMovieSequence)
 	{
 		return;
 	}
 	
 	UMovieScene* const FocusedMovieScene = FocusedMovieSequence->GetMovieScene();
-	if (!IsValid(FocusedMovieScene))
+	if (!FocusedMovieScene)
 	{
 		return;
 	}
@@ -441,7 +441,7 @@ UWorld* SNavigationToolFilterBar::GetWorld() const
 	}
 
 	UObject* const PlaybackContext = FilterBar->GetSequencer().GetPlaybackContext();
-	if (!IsValid(PlaybackContext))
+	if (!PlaybackContext)
 	{
 		return nullptr;
 	}
@@ -517,7 +517,7 @@ void SNavigationToolFilterBar::CreateFilterWidgetsFromConfig()
 	}
 
 	UNavigationToolSettings* const ToolSettings = GetMutableDefault<UNavigationToolSettings>();
-	if (!ensure(IsValid(ToolSettings)))
+	if (!ToolSettings)
 	{
 		return;
 	}
diff --git a/Engine/Plugins/Experimental/SequenceNavigator/Source/SequenceNavigator/Private/Widgets/SNavigationToolView.cpp b/Engine/Plugins/Experimental/SequenceNavigator/Source/SequenceNavigator/Private/Widgets/SNavigationToolView.cpp
--- a/Engine/Plugins/Experimental/SequenceNavigator/Source/SequenceNavigator/Private/Widgets/SNavigationToolView.cpp
+++ b/Engine/Plugins/Experimental/SequenceNavigator/Source/SequenceNavigator/Private/Widgets/SNavigationToolView.cpp
@@ -194,8 +194,7 @@ TSharedRef<SWidget> SNavigationToolView::ConstructSplitterContent()
 		+ SSplitter::Slot()
 		.Value(TAttribute<float>::CreateLambda([this]()
 			{
-				const UNavigationToolSettings* const ToolSettings = GetDefault<UNavigationToolSettings>();
-				if (ensure(IsValid(ToolSettings)))
+				if (const UNavigationToolSettings* const ToolSettings = GetDefault<UNavigationToolSettings>())
 				{
 					return ToolSettings->GetLastFilterBarSizeCoefficient();
 				}
@@ -203,8 +202,7 @@ TSharedRef<SWidget> SNavigationToolView::ConstructSplitterContent()
 			}))
 		.OnSlotResized_Lambda([this](const float InNewCoefficient)
 			{
-				UNavigationToolSettings* ToolSettings = GetMutableDefault<UNavigationToolSettings>();
-				if (ensure(IsValid(ToolSettings)))
+				if (UNavigationToolSettings* const ToolSettings = GetMutableDefault<UNavigationToolSettings>())
 				{
 					ToolSettings->SetLastFilterBarSizeCoefficient(InNewCoefficient);
 				}
diff --git a/Engine/Plugins/Experimental/SequenceNavigator/Source/SequenceNavigator/Public/Filters/NavigationToolFilterBar.h b/Engine/Plugins/Experimental/SequenceNavigator/Source/SequenceNavigator/Public/Filters/NavigationToolFilterBar.h
--- a/Engine/Plugins/Experimental/SequenceNavigator/Source/SequenceNavigator/Public/Filters/NavigationToolFilterBar.h
+++ b/Engine/Plugins/Experimental/SequenceNavigator/Source/SequenceNavigator/Public/Filters/NavigationToolFilterBar.h
@@ -260,8 +260,7 @@ private:
 	static int32 InstanceCount;
 
 	/** Do not call directly! Should only be called by FilterNodes(). */
-	bool FilterNodesRecursive(ISequencer& InSequencer
-		, INavigationToolView& InToolView
+	bool FilterNodesRecursive(INavigationToolView& InToolView
 		, const UNavigationToolSettings& InSettings
 		, const bool bInHasActiveFilter
 		, const FNavigationToolItemPtr& InStartNode);
