[Sequencer] Navigation Tool:
- Items can specify their own label color, separate from item row tint
- Only color the item icon and label for invalid items instead of the whole item row
- Playhead Column and Filter is inconsistent while scrubbing through a Root Sequence

diff --git a/Engine/Plugins/Experimental/SequenceNavigator/Source/SequenceNavigator/Private/Filters/Filters/NavigationToolFilter_Playhead.cpp b/Engine/Plugins/Experimental/SequenceNavigator/Source/SequenceNavigator/Private/Filters/Filters/NavigationToolFilter_Playhead.cpp
--- a/Engine/Plugins/Experimental/SequenceNavigator/Source/SequenceNavigator/Private/Filters/Filters/NavigationToolFilter_Playhead.cpp
+++ b/Engine/Plugins/Experimental/SequenceNavigator/Source/SequenceNavigator/Private/Filters/Filters/NavigationToolFilter_Playhead.cpp
@@ -97,7 +97,7 @@ bool FNavigationToolFilter_Playhead::PassesFilter(FNavigationToolFilterType InIt
 		return false;
 	}
 
-	return PlayheadItem->ContainsPlayhead() == EItemContainsPlayhead::ContainsPlayhead;
+	return PlayheadItem->ContainsPlayhead() != EItemContainsPlayhead::None;
 }
 
 void FNavigationToolFilter_Playhead::OnPlayEvent()
diff --git a/Engine/Plugins/Experimental/SequenceNavigator/Source/SequenceNavigator/Private/Items/NavigationToolBinding.cpp b/Engine/Plugins/Experimental/SequenceNavigator/Source/SequenceNavigator/Private/Items/NavigationToolBinding.cpp
--- a/Engine/Plugins/Experimental/SequenceNavigator/Source/SequenceNavigator/Private/Items/NavigationToolBinding.cpp
+++ b/Engine/Plugins/Experimental/SequenceNavigator/Source/SequenceNavigator/Private/Items/NavigationToolBinding.cpp
@@ -3,17 +3,17 @@
 #include "Items/NavigationToolBinding.h"
 #include "INavigationTool.h"
 #include "ISequencer.h"
+#include "Items/NavigationToolItemUtils.h"
 #include "Items/NavigationToolSequence.h"
 #include "MovieScene.h"
 #include "ScopedTransaction.h"
-#include "MVVM/Extensions/ILayerBarExtension.h"
 #include "MVVM/ObjectBindingModelStorageExtension.h"
-#include "MVVM/TrackModelStorageExtension.h"
 #include "MVVM/ViewModels/SequencerEditorViewModel.h"
 #include "MVVM/ViewModels/TrackModel.h"
 #include "NavigationToolScopedSelection.h"
 #include "Sections/MovieSceneSubSection.h"
 #include "Styling/SlateIconFinder.h"
+#include "Styling/StyleColors.h"
 #include "Utils/NavigationToolMiscUtils.h"
 #include "Widgets/Views/STableRow.h"
 
@@ -66,13 +66,18 @@ FText FNavigationToolBinding::GetDisplayName() const
 
 FText FNavigationToolBinding::GetClassName() const
 {
-	if (const UClass* const BoundObjectClass = GetBoundObjectClass())
+	if (WeakBoundObjectClass.IsValid())
 	{
-		return BoundObjectClass->GetDisplayNameText();
+		return WeakBoundObjectClass->GetDisplayNameText();
 	}
 	return FText::GetEmpty();
 }
 
+FSlateColor FNavigationToolBinding::GetItemLabelColor() const
+{
+	return IconColor;
+}
+
 FSlateIcon FNavigationToolBinding::GetIcon() const
 {
 	return Icon;
@@ -89,44 +94,9 @@ FText FNavigationToolBinding::GetIconTooltipText() const
 		, FText::FromString(LexToString(Binding.GetObjectGuid())));
 }
 
-FLinearColor FNavigationToolBinding::GetItemBrushColor() const
+FSlateColor FNavigationToolBinding::GetIconColor() const
 {
-	const TSharedPtr<ISequencer> Sequencer = Tool.GetSequencer();
-	if (!Sequencer.IsValid())
-	{
-		return FLinearColor::Red;
-	}
-
-	UMovieSceneSequence* const Sequence = GetSequence();
-	if (!Sequence)
-	{
-		return FLinearColor::Red;
-	}
-
-	const TArrayView<TWeakObjectPtr<>> BoundObjects = ResolveBoundObjects(*Sequencer, Sequence, Binding.GetObjectGuid());
-	if (BoundObjects.Num() > 0)
-	{
-		int32 NumValidObjects = 0;
-		for (const TWeakObjectPtr<>& BoundObject : BoundObjects)
-		{
-			if (BoundObject.IsValid())
-			{
-				++NumValidObjects;
-			}
-		}
-
-		if (NumValidObjects == BoundObjects.Num())
-		{
-			return FNavigationToolItem::GetItemBrushColor();
-		}
-
-		if (NumValidObjects > 0)
-		{
-			return FLinearColor::Yellow;
-		}
-	}
-
-	return FLinearColor::Red;
+	return IconColor;
 }
 
 bool FNavigationToolBinding::IsSelected(const FNavigationToolScopedSelection& InSelection) const
@@ -247,30 +217,16 @@ EItemContainsPlayhead FNavigationToolBinding::ContainsPlayhead() const
 		return EItemContainsPlayhead::None;
 	}
 
-	FTrackModelStorageExtension* const TrackModelStorage = Sequencer->GetViewModel()->CastDynamic<FTrackModelStorageExtension>();
-	if (!TrackModelStorage)
-	{
-		return EItemContainsPlayhead::None;
-	}
-
 	const FQualifiedFrameTime PlayheadTime = Sequencer->GetLocalTime();
 
-	for (UMovieSceneTrack* const Track : Binding.GetTracks())
+	for (const UMovieSceneTrack* const Track : Binding.GetTracks())
 	{
-		const TSharedPtr<FTrackModel> TrackModel = TrackModelStorage->FindModelForTrack(Track);
-		if (!TrackModel.IsValid())
+		const TArray<UMovieSceneSection*> Sections = Track->GetAllSections();
+		for (const UMovieSceneSection* const Section : Sections)
 		{
-			continue;
-		}
-
-		for (const TViewModelPtr<ILayerBarExtension>& TrackAreaModel : TrackModel->GetTrackAreaModelListAs<ILayerBarExtension>())
-		{
-			if (const TViewModelPtr<ILayerBarExtension> BarModel = TrackAreaModel.ImplicitCast())
+			if (Section->GetTrueRange().Contains(PlayheadTime.Time.FrameNumber))
 			{
-				if (BarModel->GetLayerBarRange().Contains(PlayheadTime.Time.FrameNumber))
-				{
-					return EItemContainsPlayhead::ContainsPlayhead;
-				}
+				return EItemContainsPlayhead::ContainsPlayhead;
 			}
 		}
 	}
@@ -278,29 +234,6 @@ EItemContainsPlayhead FNavigationToolBinding::ContainsPlayhead() const
 	return EItemContainsPlayhead::None;
 }
 
-const UClass* FNavigationToolBinding::GetBoundObjectClass() const
-{
-	const TSharedPtr<ISequencer> Sequencer = Tool.GetSequencer();
-	if (!Sequencer.IsValid())
-	{
-		return nullptr;
-	}
-
-	UMovieSceneSequence* const Sequence = GetSequence();
-	if (!Sequence)
-	{
-		return nullptr;
-	}
-
-	UMovieScene* const MovieScene = Sequence->GetMovieScene();
-	if (!MovieScene)
-	{
-		return nullptr;
-	}
-
-	return MovieSceneHelpers::GetBoundObjectClass(Sequence, Binding.GetObjectGuid());
-}
-
 UObject* FNavigationToolBinding::GetCachedBoundObject() const
 {
 	return WeakBoundObject.Get();
@@ -323,8 +256,12 @@ UObject* FNavigationToolBinding::CacheBoundObject()
 	const TArrayView<TWeakObjectPtr<>> WeakBoundObjects = ResolveBoundObjects(*Sequencer, Sequence, Binding.GetObjectGuid());
 	WeakBoundObject = WeakBoundObjects.IsEmpty() ? nullptr : WeakBoundObjects[0];
 
-	const UClass* const BoundObjectClass = MovieSceneHelpers::GetBoundObjectClass(Sequence, Binding.GetObjectGuid());
-	Icon = FSlateIconFinder::FindIconForClass(BoundObjectClass);
+	WeakBoundObjectClass = MovieSceneHelpers::GetBoundObjectClass(Sequence, Binding.GetObjectGuid());
+
+	Icon = FSlateIconFinder::FindIconForClass(WeakBoundObjectClass.Get());
+
+	const FSlateColor DefaultColor = FNavigationToolItem::GetItemLabelColor();
+	IconColor = GetItemBindingColor(*Sequencer, *Sequence, Binding.GetObjectGuid(), DefaultColor);
 
 	return WeakBoundObject.Get();
 }
diff --git a/Engine/Plugins/Experimental/SequenceNavigator/Source/SequenceNavigator/Private/Items/NavigationToolComponent.cpp b/Engine/Plugins/Experimental/SequenceNavigator/Source/SequenceNavigator/Private/Items/NavigationToolComponent.cpp
--- a/Engine/Plugins/Experimental/SequenceNavigator/Source/SequenceNavigator/Private/Items/NavigationToolComponent.cpp
+++ b/Engine/Plugins/Experimental/SequenceNavigator/Source/SequenceNavigator/Private/Items/NavigationToolComponent.cpp
@@ -67,7 +67,7 @@ TSharedRef<SWidget> FNavigationToolComponent::GenerateLabelWidget(const TSharedR
 	return SNew(SNavigationToolLabelComponent, SharedThis(this), InRow);
 }
 
-FLinearColor FNavigationToolComponent::GetItemBrushColor() const
+FLinearColor FNavigationToolComponent::GetItemTintColor() const
 {
 	return FStyleColors::White25.GetSpecifiedColor();
 }
diff --git a/Engine/Plugins/Experimental/SequenceNavigator/Source/SequenceNavigator/Private/Items/NavigationToolItem.cpp b/Engine/Plugins/Experimental/SequenceNavigator/Source/SequenceNavigator/Private/Items/NavigationToolItem.cpp
--- a/Engine/Plugins/Experimental/SequenceNavigator/Source/SequenceNavigator/Private/Items/NavigationToolItem.cpp
+++ b/Engine/Plugins/Experimental/SequenceNavigator/Source/SequenceNavigator/Private/Items/NavigationToolItem.cpp
@@ -224,6 +224,11 @@ FNavigationToolItemId FNavigationToolItem::GetItemId() const
 	return ItemId;
 }
 
+FSlateColor FNavigationToolItem::GetItemLabelColor() const
+{
+	return FStyleColors::Foreground;
+}
+
 const FSlateBrush* FNavigationToolItem::GetIconBrush() const
 {
 	const FSlateIcon Icon = FNavigationToolExtender::FindOverrideIcon(SharedThis(this));
@@ -340,7 +345,7 @@ FReply FNavigationToolItem::AcceptDrop(const FDragDropEvent& InDragDropEvent, EI
 	return OutReply;
 }
 
-FLinearColor FNavigationToolItem::GetItemBrushColor() const
+FLinearColor FNavigationToolItem::GetItemTintColor() const
 {
 	return FStyleColors::White.GetSpecifiedColor();
 }
diff --git a/Engine/Plugins/Experimental/SequenceNavigator/Source/SequenceNavigator/Private/Items/NavigationToolItemUtils.cpp b/Engine/Plugins/Experimental/SequenceNavigator/Source/SequenceNavigator/Private/Items/NavigationToolItemUtils.cpp
--- a/Engine/Plugins/Experimental/SequenceNavigator/Source/SequenceNavigator/Private/Items/NavigationToolItemUtils.cpp
+++ b/Engine/Plugins/Experimental/SequenceNavigator/Source/SequenceNavigator/Private/Items/NavigationToolItemUtils.cpp
@@ -12,6 +12,7 @@
 #include "Sections/MovieSceneCinematicShotSection.h"
 #include "Sections/MovieSceneSubSection.h"
 #include "TrackEditors/SubTrackEditorBase.h"
+#include "Utils/NavigationToolMiscUtils.h"
 
 #define LOCTEXT_NAMESPACE "NavigationToolItemUtils"
 
@@ -175,6 +176,38 @@ void AppendSequenceDisplayNameDirtyStatus(FText& InOutDisplayName
 	}
 }
 
+FSlateColor GetItemBindingColor(const ISequencer& InSequencer
+	, UMovieSceneSequence& InSequence
+	, const FGuid& InObjectGuid
+	, const FSlateColor& DefaultColor)
+{
+	const TArrayView<TWeakObjectPtr<>> BoundObjects = ResolveBoundObjects(InSequencer, &InSequence, InObjectGuid);
+	if (BoundObjects.Num() > 0)
+	{
+		int32 NumValidObjects = 0;
+
+		for (const TWeakObjectPtr<>& BoundObject : BoundObjects)
+		{
+			if (BoundObject.IsValid())
+			{
+				++NumValidObjects;
+			}
+		}
+
+		if (NumValidObjects == BoundObjects.Num())
+		{
+			return DefaultColor;
+		}
+
+		if (NumValidObjects > 0)
+		{
+			return FStyleColors::Warning;
+		}
+	}
+
+	return FStyleColors::Error;
+}
+
 } // namespace UE::SequenceNavigator
 
 #undef LOCTEXT_NAMESPACE
diff --git a/Engine/Plugins/Experimental/SequenceNavigator/Source/SequenceNavigator/Private/Items/NavigationToolSequence.cpp b/Engine/Plugins/Experimental/SequenceNavigator/Source/SequenceNavigator/Private/Items/NavigationToolSequence.cpp
--- a/Engine/Plugins/Experimental/SequenceNavigator/Source/SequenceNavigator/Private/Items/NavigationToolSequence.cpp
+++ b/Engine/Plugins/Experimental/SequenceNavigator/Source/SequenceNavigator/Private/Items/NavigationToolSequence.cpp
@@ -76,8 +76,9 @@ void FNavigationToolSequence::FindChildren(TArray<FNavigationToolItemPtr>& OutCh
 	}
 
 	const TSharedRef<FNavigationToolSequence> SharedThisRef = SharedThis(this);
-
+	const TSharedPtr<FNavigationToolProvider> Provider = GetProvider();
 	const TArray<UMovieSceneTrack*> MovieSceneTracks = MovieScene->GetTracks();
+
 	for (UMovieSceneTrack* const Track : MovieSceneTracks)
 	{
 		FNavigationToolItemPtr NewItem;
@@ -524,19 +525,35 @@ EItemContainsPlayhead FNavigationToolSequence::ContainsPlayhead() const
 		return EItemContainsPlayhead::None;
 	}
 
-	const FQualifiedFrameTime PlayheadTime = Sequencer->GetLocalTime();
+	const UMovieSceneSequence* const Sequence = WeakSequence.Get();
+	if (!IsValid(Sequence))
+	{
+		return EItemContainsPlayhead::None;
+	}
 
-	if (UMovieSceneSubSection* const SubSection = GetSubSection())
+	UMovieSceneSequence* const FocusedSequence = Sequencer->GetFocusedMovieSceneSequence();
+	if (!IsValid(FocusedSequence))
 	{
-		return SubSection->GetTrueRange().Contains(PlayheadTime.Time.FrameNumber)
-			? EItemContainsPlayhead::ContainsPlayhead : EItemContainsPlayhead::None;
+		return EItemContainsPlayhead::None;
 	}
 
-	if (UMovieSceneSequence* const RootSequence = Sequencer->GetRootMovieSceneSequence())
+	const FQualifiedFrameTime PlayheadTime = Sequencer->GetLocalTime();
+
+	if (Sequence == FocusedSequence)
+	{
+		UMovieScene* const SequenceMovieScene = Sequence->GetMovieScene();
+		if (IsValid(SequenceMovieScene))
+		{
+			return SequenceMovieScene->GetPlaybackRange().Contains(PlayheadTime.Time.FrameNumber)
+				? EItemContainsPlayhead::ContainsPlayhead : EItemContainsPlayhead::None;
+		}
+	}
+	else
 	{
-		if (UMovieScene* const RootSequenceMovieScene = RootSequence->GetMovieScene())
+		UMovieSceneSubSection* const SubSection = GetSubSection();
+		if (IsValid(SubSection))
 		{
-			return RootSequenceMovieScene->GetPlaybackRange().Contains(PlayheadTime.Time.FrameNumber)
+			return SubSection->GetTrueRange().Contains(PlayheadTime.Time.FrameNumber)
 				? EItemContainsPlayhead::ContainsPlayhead : EItemContainsPlayhead::None;
 		}
 	}
diff --git a/Engine/Plugins/Experimental/SequenceNavigator/Source/SequenceNavigator/Private/Items/NavigationToolTrack.cpp b/Engine/Plugins/Experimental/SequenceNavigator/Source/SequenceNavigator/Private/Items/NavigationToolTrack.cpp
--- a/Engine/Plugins/Experimental/SequenceNavigator/Source/SequenceNavigator/Private/Items/NavigationToolTrack.cpp
+++ b/Engine/Plugins/Experimental/SequenceNavigator/Source/SequenceNavigator/Private/Items/NavigationToolTrack.cpp
@@ -1,10 +1,11 @@
-﻿// Copyright Epic Games, Inc. All Rights Reserved.
+// Copyright Epic Games, Inc. All Rights Reserved.
 
 #include "Items/NavigationToolTrack.h"
 #include "INavigationTool.h"
 #include "ISequencer.h"
 #include "ISequencerTrackEditor.h"
 #include "Items/NavigationToolItemUtils.h"
+#include "MovieScene.h"
 #include "MovieSceneSequence.h"
 #include "MovieSceneTrack.h"
 #include "MVVM/TrackModelStorageExtension.h"
@@ -214,24 +215,30 @@ EItemContainsPlayhead FNavigationToolTrack::ContainsPlayhead() const
 		return EItemContainsPlayhead::None;
 	}
 
+	const FQualifiedFrameTime PlayheadTime = Sequencer->GetLocalTime();
+
 	if (WeakSection.IsValid())
 	{
-		const FQualifiedFrameTime PlayheadTime = Sequencer->GetLocalTime();
 		return WeakSection.Pin()->GetTrueRange().Contains(PlayheadTime.Time.FrameNumber)
 			? EItemContainsPlayhead::ContainsPlayhead : EItemContainsPlayhead::None;
 	}
 
-	const ENavigationToolCompareState State = CompareChildrenItemState<IPlayheadExtension>(this,
-		[](const IPlayheadExtension* const InPlayheadItem)
-			{
-				return InPlayheadItem->ContainsPlayhead() == EItemContainsPlayhead::ContainsPlayhead;
-			},
-		[](const IPlayheadExtension* const InPlayheadItem)
-			{
-				return InPlayheadItem->ContainsPlayhead() != EItemContainsPlayhead::ContainsPlayhead;
-			});
+	UMovieSceneTrack* const Track = GetTrack();
+	if (!Track)
+	{
+		return EItemContainsPlayhead::None;
+	}
+
+	const TArray<UMovieSceneSection*> Sections = Track->GetAllSections();
+	for (const UMovieSceneSection* const Section : Sections)
+	{
+		if (Section->GetTrueRange().Contains(PlayheadTime.Time.FrameNumber))
+		{
+			return EItemContainsPlayhead::ContainsPlayhead;
+		}
+	}
 
-	return static_cast<EItemContainsPlayhead>(State);
+	return EItemContainsPlayhead::None;
 }
 
 void FNavigationToolTrack::OnTrackObjectChanged()
diff --git a/Engine/Plugins/Experimental/SequenceNavigator/Source/SequenceNavigator/Private/NavigationToolView.cpp b/Engine/Plugins/Experimental/SequenceNavigator/Source/SequenceNavigator/Private/NavigationToolView.cpp
--- a/Engine/Plugins/Experimental/SequenceNavigator/Source/SequenceNavigator/Private/NavigationToolView.cpp
+++ b/Engine/Plugins/Experimental/SequenceNavigator/Source/SequenceNavigator/Private/NavigationToolView.cpp
@@ -348,11 +348,14 @@ void FNavigationToolView::RequestRefresh()
 
 void FNavigationToolView::Refresh()
 {
-	UpdateRootVisibleItems();
-
+	// Filter items before doing anything else so we can reliably use the filter data cache.
+	// For example, in cases where a FNavigationToolAddItem is executed and a new item is added
+	// to the tree, UpdateRootVisibleItems() below uses the filter data to show/hide items.
 	bFilterUpdateRequested = true;
 	UpdateFilters();
 
+	UpdateRootVisibleItems();
+
 	UpdateItemExpansions();
 	
 	if (ToolViewWidget.IsValid())
@@ -534,8 +537,7 @@ void FNavigationToolView::GetChildrenOfItem(const FNavigationToolItemPtr& InItem
 			continue;
 		}
 
-		const bool bShouldShowItem = ShouldShowItem(ChildItem, true, InViewMode);
-		if (bShouldShowItem)
+		if (const bool bShouldShowItem = ShouldShowItem(ChildItem, true, InViewMode))
 		{
 			// If the current item is visible in outliner, add it to the children
 			OutChildren.Add(ChildItem);
@@ -580,7 +582,7 @@ FLinearColor FNavigationToolView::GetItemBrushColor(const FNavigationToolItemPtr
 {
 	if (InItem.IsValid())
 	{
-		FLinearColor OutColor = InItem->GetItemBrushColor();
+		FLinearColor OutColor = InItem->GetItemTintColor();
 
 		// If NextSelectedItemIntoView is valid, it means we're scrolling items into view with Next/Previous, 
 		// so Make everything that's not the Current Item a bit more translucent to make the Current Item stand out
diff --git a/Engine/Plugins/Experimental/SequenceNavigator/Source/SequenceNavigator/Private/Widgets/Columns/SNavigationToolLabelComponent.h b/Engine/Plugins/Experimental/SequenceNavigator/Source/SequenceNavigator/Private/Widgets/Columns/SNavigationToolLabelComponent.h
--- a/Engine/Plugins/Experimental/SequenceNavigator/Source/SequenceNavigator/Private/Widgets/Columns/SNavigationToolLabelComponent.h
+++ b/Engine/Plugins/Experimental/SequenceNavigator/Source/SequenceNavigator/Private/Widgets/Columns/SNavigationToolLabelComponent.h
@@ -4,10 +4,13 @@
 
 #include "SNavigationToolLabelItem.h"
 
+struct FInlineEditableTextBlockStyle;
+
 namespace UE::SequenceNavigator
 {
 
 class FNavigationToolComponent;
+class SNavigationToolTreeRow;
 
 class SNavigationToolLabelComponent : public SNavigationToolLabelItem
 {
@@ -19,7 +22,9 @@ public:
 		, const TSharedRef<FNavigationToolComponent>& InItem
 		, const TSharedRef<SNavigationToolTreeRow>& InRowWidget);
 
+	//~ Begin SNavigationToolLabelItem
 	virtual const FInlineEditableTextBlockStyle* GetTextBlockStyle() const override;
+	//~ End SNavigationToolLabelItem
 };
 
 } // namespace UE::SequenceNavigator
diff --git a/Engine/Plugins/Experimental/SequenceNavigator/Source/SequenceNavigator/Private/Widgets/Columns/SNavigationToolLabelItem.cpp b/Engine/Plugins/Experimental/SequenceNavigator/Source/SequenceNavigator/Private/Widgets/Columns/SNavigationToolLabelItem.cpp
--- a/Engine/Plugins/Experimental/SequenceNavigator/Source/SequenceNavigator/Private/Widgets/Columns/SNavigationToolLabelItem.cpp
+++ b/Engine/Plugins/Experimental/SequenceNavigator/Source/SequenceNavigator/Private/Widgets/Columns/SNavigationToolLabelItem.cpp
@@ -86,6 +86,29 @@ void SNavigationToolLabelItem::Construct(const FArguments& InArgs
 	];
 }
 
+FSlateColor SNavigationToolLabelItem::GetForegroundColor() const
+{
+	const TSharedPtr<INavigationToolView> View = WeakView.Pin();
+	if (!View.IsValid())
+	{
+		return SCompoundWidget::GetForegroundColor();
+	}
+
+	const TSharedPtr<ISequencer> Sequencer = View->GetSequencer();
+	if (!Sequencer.IsValid())
+	{
+		return SCompoundWidget::GetForegroundColor();
+	}
+
+	const TSharedPtr<INavigationToolItem> Item = WeakItem.Pin();
+	if (!Item.IsValid())
+	{
+		return SCompoundWidget::GetForegroundColor();
+	}
+
+	return Item->GetItemLabelColor();
+}
+
 bool SNavigationToolLabelItem::IsReadOnly() const
 {
 	const FNavigationToolItemPtr Item = WeakItem.Pin();
diff --git a/Engine/Plugins/Experimental/SequenceNavigator/Source/SequenceNavigator/Private/Widgets/Columns/SNavigationToolLabelItem.h b/Engine/Plugins/Experimental/SequenceNavigator/Source/SequenceNavigator/Private/Widgets/Columns/SNavigationToolLabelItem.h
--- a/Engine/Plugins/Experimental/SequenceNavigator/Source/SequenceNavigator/Private/Widgets/Columns/SNavigationToolLabelItem.h
+++ b/Engine/Plugins/Experimental/SequenceNavigator/Source/SequenceNavigator/Private/Widgets/Columns/SNavigationToolLabelItem.h
@@ -7,6 +7,7 @@
 #include "Widgets/SCompoundWidget.h"
 
 class SInlineEditableTextBlock;
+struct FSlateColor;
 
 namespace UE::SequenceNavigator
 {
@@ -44,7 +45,9 @@ public:
 
 	virtual const FInlineEditableTextBlockStyle* GetTextBlockStyle() const;
 
-private:
+	virtual FSlateColor GetForegroundColor() const;
+
+protected:
 	FNavigationToolItemWeakPtr WeakItem;
 
 	TWeakPtr<INavigationToolView> WeakView;
diff --git a/Engine/Plugins/Experimental/SequenceNavigator/Source/SequenceNavigator/Public/Items/INavigationToolItem.h b/Engine/Plugins/Experimental/SequenceNavigator/Source/SequenceNavigator/Public/Items/INavigationToolItem.h
--- a/Engine/Plugins/Experimental/SequenceNavigator/Source/SequenceNavigator/Public/Items/INavigationToolItem.h
+++ b/Engine/Plugins/Experimental/SequenceNavigator/Source/SequenceNavigator/Public/Items/INavigationToolItem.h
@@ -420,6 +420,12 @@ public:
 	/** Gets the Class/Type of this Item (e.g. for Items that represent UObjects, it will be the UObject class) */
 	virtual FText GetClassName() const = 0;
 
+	/** Gets the color to use for the item label */
+	virtual FSlateColor GetItemLabelColor() const = 0;
+
+	/** Gets the color to use to tint the entire item row (all column content) */
+	virtual FLinearColor GetItemTintColor() const = 0;
+
 	/** Gets the slate icon for the item */
 	virtual FSlateIcon GetIcon() const = 0;
 
@@ -536,9 +542,6 @@ public:
 	/** Processes the Drag and Drop Event for this Item */
 	virtual FReply AcceptDrop(const FDragDropEvent& InDragDropEvent, const EItemDropZone InDropZone) = 0;
 
-	/** Gets the Text Color to use for the Item in the Navigation Tool (e.g. Components should be grayed out a bit) */
-	virtual FLinearColor GetItemBrushColor() const = 0;
-
 	/** Whether Ignoring Pending Kill. Useful to get underlying UObjects that are pending kill and get the pointer to it and not a null value */
 	bool IsIgnoringPendingKill() const { return HasAllFlags(ENavigationToolItemFlags::IgnorePendingKill); }
 
diff --git a/Engine/Plugins/Experimental/SequenceNavigator/Source/SequenceNavigator/Public/Items/NavigationToolBinding.h b/Engine/Plugins/Experimental/SequenceNavigator/Source/SequenceNavigator/Public/Items/NavigationToolBinding.h
--- a/Engine/Plugins/Experimental/SequenceNavigator/Source/SequenceNavigator/Public/Items/NavigationToolBinding.h
+++ b/Engine/Plugins/Experimental/SequenceNavigator/Source/SequenceNavigator/Public/Items/NavigationToolBinding.h
@@ -9,6 +9,7 @@
 #include "MVVM/ViewModels/ObjectBindingModel.h"
 #include "NavigationToolItem.h"
 #include "SequencerCoreFwd.h"
+#include "Styling/SlateColor.h"
 #include "Textures/SlateIcon.h"
 
 namespace UE::Sequencer
@@ -55,9 +56,10 @@ public:
 
 	virtual FText GetDisplayName() const override;
 	virtual FText GetClassName() const override;
+	virtual FSlateColor GetItemLabelColor() const override;
 	virtual FSlateIcon GetIcon() const override;
 	virtual FText GetIconTooltipText() const override;
-	virtual FLinearColor GetItemBrushColor() const override;
+	virtual FSlateColor GetIconColor() const override;
 
 	virtual bool IsSelected(const FNavigationToolScopedSelection& InSelection) const override;
 	virtual void Select(FNavigationToolScopedSelection& InSelection) const override;
@@ -85,8 +87,6 @@ public:
 	virtual EItemContainsPlayhead ContainsPlayhead() const override;
 	//~ End IPlayheadExtension
 
-	const UClass* GetBoundObjectClass() const;
-
 	const FMovieSceneBinding& GetBinding() const;
 
 	/** @return The Sequence this binding belongs to */
@@ -105,9 +105,11 @@ protected:
 
 	FMovieSceneBinding Binding;
 
+	TWeakObjectPtr<const UClass> WeakBoundObjectClass;
 	TWeakObjectPtr<> WeakBoundObject;
 
 	FSlateIcon Icon;
+	FSlateColor IconColor;
 };
 
 } // namespace UE::SequenceNavigator
diff --git a/Engine/Plugins/Experimental/SequenceNavigator/Source/SequenceNavigator/Public/Items/NavigationToolComponent.h b/Engine/Plugins/Experimental/SequenceNavigator/Source/SequenceNavigator/Public/Items/NavigationToolComponent.h
--- a/Engine/Plugins/Experimental/SequenceNavigator/Source/SequenceNavigator/Public/Items/NavigationToolComponent.h
+++ b/Engine/Plugins/Experimental/SequenceNavigator/Source/SequenceNavigator/Public/Items/NavigationToolComponent.h
@@ -34,7 +34,7 @@ public:
 	virtual bool CanReceiveParentVisibilityPropagation() const override { return false; }
 	virtual TSharedRef<SWidget> GenerateLabelWidget(const TSharedRef<SNavigationToolTreeRow>& InRow) override;
 
-	virtual FLinearColor GetItemBrushColor() const override;
+	virtual FLinearColor GetItemTintColor() const override;
 	virtual TArray<FName> GetTags() const override;
 
 	virtual bool ShowVisibility() const override { return true; }
diff --git a/Engine/Plugins/Experimental/SequenceNavigator/Source/SequenceNavigator/Public/Items/NavigationToolItem.h b/Engine/Plugins/Experimental/SequenceNavigator/Source/SequenceNavigator/Public/Items/NavigationToolItem.h
--- a/Engine/Plugins/Experimental/SequenceNavigator/Source/SequenceNavigator/Public/Items/NavigationToolItem.h
+++ b/Engine/Plugins/Experimental/SequenceNavigator/Source/SequenceNavigator/Public/Items/NavigationToolItem.h
@@ -56,6 +56,8 @@ public:
 	virtual FText GetClassName() const override { return FText::GetEmpty(); }
 	virtual ENavigationToolItemViewMode GetSupportedViewModes(const INavigationToolView& InToolView) const override;
 	virtual FNavigationToolItemId GetItemId() const override final;
+	virtual FSlateColor GetItemLabelColor() const override;
+	virtual FLinearColor GetItemTintColor() const override;
 	virtual const FSlateBrush* GetIconBrush() const override;
 	virtual FSlateColor GetIconColor() const override;
 	virtual TSharedRef<SWidget> GenerateLabelWidget(const TSharedRef<SNavigationToolTreeRow>& InRow) override;
@@ -74,7 +76,6 @@ public:
 	virtual TArray<FName> GetTags() const override { return TArray<FName>(); };
 	virtual TOptional<EItemDropZone> CanAcceptDrop(const FDragDropEvent& InDragDropEvent, EItemDropZone InDropZone) override;
 	virtual FReply AcceptDrop(const FDragDropEvent& InDragDropEvent, EItemDropZone InDropZone) override;
-	virtual FLinearColor GetItemBrushColor() const override;
 	virtual FOnRenameAction& OnRenameAction() override { return OnRenameActionDelegate; }
 	virtual FOnExpansionChanged& OnExpansionChanged() override { return OnExpansionChangedDelegate; }
 	virtual bool IsExpanded() const override;
diff --git a/Engine/Plugins/Experimental/SequenceNavigator/Source/SequenceNavigator/Public/Items/NavigationToolItemUtils.h b/Engine/Plugins/Experimental/SequenceNavigator/Source/SequenceNavigator/Public/Items/NavigationToolItemUtils.h
--- a/Engine/Plugins/Experimental/SequenceNavigator/Source/SequenceNavigator/Public/Items/NavigationToolItemUtils.h
+++ b/Engine/Plugins/Experimental/SequenceNavigator/Source/SequenceNavigator/Public/Items/NavigationToolItemUtils.h
@@ -8,6 +8,7 @@
 #include "NavigationToolDefines.h"
 #include "NavigationToolSequence.h"
 
+class ISequencer;
 class UMovieSceneMetaData;
 class UMovieSceneSubSection;
 
@@ -236,4 +237,9 @@ ENavigationToolCompareState CompareArrayStateSimple(const TArray<InItemType*>& I
 	return ENavigationToolCompareState::PartialTrue;
 }
 
+FSlateColor GetItemBindingColor(const ISequencer& InSequencer
+	, UMovieSceneSequence& InSequence
+	, const FGuid& InObjectGuid
+	, const FSlateColor& DefaultColor = FSlateColor::UseForeground());
+
 } // namespace UE::SequenceNavigator
