[SequenceNavigator] Fix slow undo/redo

diff --git a/Engine/Plugins/Experimental/SequenceNavigator/Source/SequenceNavigator/Private/Filters/NavigationToolFilterBar.cpp b/Engine/Plugins/Experimental/SequenceNavigator/Source/SequenceNavigator/Private/Filters/NavigationToolFilterBar.cpp
--- a/Engine/Plugins/Experimental/SequenceNavigator/Source/SequenceNavigator/Private/Filters/NavigationToolFilterBar.cpp
+++ b/Engine/Plugins/Experimental/SequenceNavigator/Source/SequenceNavigator/Private/Filters/NavigationToolFilterBar.cpp
@@ -513,7 +513,7 @@ bool FNavigationToolFilterBar::SetFilterActive(const TSharedRef<FNavigationToolF
 
 	const bool bNewActive = InFilter->IsInverseFilter() ? !bInActive : bInActive;
 
-	FSequencerFilterBarConfig& Config = ToolSettings->FindOrAddFilterBar(GetIdentifier(), true);
+	FSequencerFilterBarConfig& Config = ToolSettings->FindOrAddFilterBar(GetIdentifier(), /*bInSaveConfig=*/false);
 
 	const FString FilterName = InFilter->GetDisplayName().ToString();
 	const bool bSuccess = (Config.IsFilterActive(FilterName) == bInActive)
@@ -585,7 +585,7 @@ void FNavigationToolFilterBar::ActivateCommonFilters(const bool bInActivate
 		return;
 	}
 
-	FSequencerFilterBarConfig& Config = ToolSettings->FindOrAddFilterBar(GetIdentifier(), false);
+	FSequencerFilterBarConfig& Config = ToolSettings->FindOrAddFilterBar(GetIdentifier(), /*bInSaveConfig=*/false);
 
 	bool bNeedsSave = false;
 
@@ -789,19 +789,13 @@ bool FNavigationToolFilterBar::IsFilterEnabled(const TSharedRef<FNavigationToolF
 
 bool FNavigationToolFilterBar::SetFilterEnabled(const TSharedRef<FNavigationToolFilter> InFilter, const bool bInEnabled, const bool bInRequestFilterUpdate)
 {
-	const TSharedPtr<ISequencer> Sequencer = Tool.GetSequencer();
-	if (!Sequencer.IsValid())
-	{
-		return false;
-	}
-
 	UNavigationToolSettings* const ToolSettings = GetMutableDefault<UNavigationToolSettings>();
 	if (!ToolSettings)
 	{
 		return false;
 	}
 
-	FSequencerFilterBarConfig& Config = ToolSettings->FindOrAddFilterBar(GetIdentifier(), true);
+	FSequencerFilterBarConfig& Config = ToolSettings->FindOrAddFilterBar(GetIdentifier(), /*bInSaveConfig=*/false);
 
 	const FString FilterName = InFilter->GetDisplayName().ToString();
 	const bool bSuccess = Config.SetFilterEnabled(FilterName, bInEnabled);
@@ -833,12 +827,6 @@ void FNavigationToolFilterBar::EnableFilters(const bool bInEnable
 	, const TArray<TSharedRef<FFilterCategory>> InMatchCategories
 	, const TArray<TSharedRef<FNavigationToolFilter>> InExceptions)
 {
-	const TSharedPtr<ISequencer> Sequencer = Tool.GetSequencer();
-	if (!Sequencer.IsValid())
-	{
-		return;
-	}
-
 	UNavigationToolSettings* const ToolSettings = GetMutableDefault<UNavigationToolSettings>();
 	if (!ToolSettings)
 	{
@@ -995,7 +983,7 @@ bool FNavigationToolFilterBar::AddCustomTextFilter(const TSharedRef<FNavigationT
 	{
 		if (UNavigationToolSettings* const ToolSettings = GetMutableDefault<UNavigationToolSettings>())
 		{
-			FSequencerFilterBarConfig& Config = ToolSettings->FindOrAddFilterBar(GetIdentifier(), false);
+			FSequencerFilterBarConfig& Config = ToolSettings->FindOrAddFilterBar(GetIdentifier(), /*bInSaveConfig=*/false);
 
 			if (Config.AddCustomTextFilter(InFilter->CreateCustomTextFilterData()))
 			{
@@ -1020,7 +1008,7 @@ bool FNavigationToolFilterBar::RemoveCustomTextFilter(const TSharedRef<FNavigati
 	{
 		if (UNavigationToolSettings* const ToolSettings = GetMutableDefault<UNavigationToolSettings>())
 		{
-			FSequencerFilterBarConfig& Config = ToolSettings->FindOrAddFilterBar(GetIdentifier(), false);
+			FSequencerFilterBarConfig& Config = ToolSettings->FindOrAddFilterBar(GetIdentifier(), /*bInSaveConfig=*/false);
 
 			const FString FilterName = InFilter->GetDisplayName().ToString();
 			if (Config.RemoveCustomTextFilter(FilterName))
@@ -1037,19 +1025,15 @@ bool FNavigationToolFilterBar::RemoveCustomTextFilter(const TSharedRef<FNavigati
 
 void FNavigationToolFilterBar::ActivateCustomTextFilters(const bool bInActivate, const TArray<TSharedRef<FNavigationToolFilter_CustomText>> InExceptions)
 {
-	const TSharedPtr<ISequencer> Sequencer = Tool.GetSequencer();
-	if (!Sequencer.IsValid())
-	{
-		return;
-	}
-
 	UNavigationToolSettings* const ToolSettings = GetMutableDefault<UNavigationToolSettings>();
 	if (!ToolSettings)
 	{
 		return;
 	}
 
-	FSequencerFilterBarConfig& Config = ToolSettings->FindOrAddFilterBar(GetIdentifier(), false);
+	FSequencerFilterBarConfig& Config = ToolSettings->FindOrAddFilterBar(GetIdentifier(), /*bInSaveConfig=*/false);
+
+	bool bNeedsSave = false;
 
 	for (const TSharedRef<FNavigationToolFilter_CustomText>& CustomTextFilter : CustomTextFilters)
 	{
@@ -1064,32 +1048,34 @@ void FNavigationToolFilterBar::ActivateCustomTextFilters(const bool bInActivate,
 					CustomTextFilter->ActiveStateChanged(false);
 				}
 
-				const ENavigationToolFilterChange FilterChangeType = bInActivate ? ENavigationToolFilterChange::Activate : ENavigationToolFilterChange::Deactivate;
+				const ENavigationToolFilterChange FilterChangeType = bInActivate
+					? ENavigationToolFilterChange::Activate : ENavigationToolFilterChange::Deactivate;
 				BroadcastFiltersChanged(CustomTextFilter, FilterChangeType);
+
+				bNeedsSave = true;
 			}
 		}
 	}
 
-	ToolSettings->SaveConfig();
+	if (bNeedsSave)
+	{
+		ToolSettings->SaveConfig();
+	}
 
 	RequestFilterUpdate();
 }
 
 void FNavigationToolFilterBar::EnableCustomTextFilters(const bool bInEnable, const TArray<TSharedRef<FNavigationToolFilter_CustomText>> InExceptions)
 {
-	const TSharedPtr<ISequencer> Sequencer = Tool.GetSequencer();
-	if (!Sequencer.IsValid())
-	{
-		return;
-	}
-
 	UNavigationToolSettings* const ToolSettings = GetMutableDefault<UNavigationToolSettings>();
 	if (!ToolSettings)
 	{
 		return;
 	}
 
-	FSequencerFilterBarConfig& Config = ToolSettings->FindOrAddFilterBar(GetIdentifier(), false);
+	FSequencerFilterBarConfig& Config = ToolSettings->FindOrAddFilterBar(GetIdentifier(), /*bInSaveConfig=*/false);
+
+	bool bNeedsSave = false;
 
 	for (const TSharedRef<FNavigationToolFilter_CustomText>& CustomTextFilter : CustomTextFilters)
 	{
@@ -1107,11 +1093,16 @@ void FNavigationToolFilterBar::EnableCustomTextFilters(const bool bInEnable, con
 				const ENavigationToolFilterChange FilterChangeType = bInEnable
 					? ENavigationToolFilterChange::Enable : ENavigationToolFilterChange::Disable;
 				BroadcastFiltersChanged(CustomTextFilter, FilterChangeType);
+
+				bNeedsSave = true;
 			}
 		}
 	}
 
-	ToolSettings->SaveConfig();
+	if (bNeedsSave)
+	{
+		ToolSettings->SaveConfig();
+	}
 
 	RequestFilterUpdate();
 }
@@ -1158,23 +1149,11 @@ TSharedRef<FFilterCategory> FNavigationToolFilterBar::GetMiscCategory() const
 
 bool FNavigationToolFilterBar::PassesAnyCommonFilter(const FNavigationToolItemPtr& InNode)
 {
-	const TSharedPtr<ISequencer> Sequencer = Tool.GetSequencer();
-	if (!Sequencer.IsValid())
-	{
-		return true;
-	}
-
-	UMovieSceneSequence* const FocusedSequence = Sequencer->GetFocusedMovieSceneSequence();
-	if (!FocusedSequence)
-	{
-		return true;
-	}
-
 	bool bPassedAnyFilters = false;
 	bool bAnyFilterActive = false;
 
 	// Only one common filter needs to pass for this node to be included in the filtered set
-	CommonFilters->ForEachFilter([this, &InNode, FocusedSequence, &bPassedAnyFilters, &bAnyFilterActive]
+	CommonFilters->ForEachFilter([this, &InNode, &bPassedAnyFilters, &bAnyFilterActive]
 		(const TSharedRef<FNavigationToolFilter>& InFilter)
 		{
 			if (IsFilterActive(InFilter))
@@ -1219,12 +1198,6 @@ const FNavigationToolFilterData& FNavigationToolFilterBar::FilterNodes()
 {
 	//SCOPE_LOG_TIME_IN_SECONDS(TEXT("FNavigationToolFilterBar::FilterNodes()"), nullptr);
 
-	const TSharedPtr<ISequencer> Sequencer = Tool.GetSequencer();
-	if (!Sequencer.IsValid())
-	{
-		return FilterData;
-	}
-
 	const TSharedPtr<INavigationToolView> ToolView = Tool.GetMostRecentToolView();
 	if (!ToolView.IsValid())
 	{
@@ -1618,9 +1591,8 @@ bool FNavigationToolFilterBar::TryCreateCustomTextFilter(const FCustomTextFilter
 		return false;
 	}
 
-	FSequencerFilterBarConfig& Config = ToolSettings->FindOrAddFilterBar(GetIdentifier(), false);
+	FSequencerFilterBarConfig& Config = ToolSettings->FindOrAddFilterBar(GetIdentifier(), /*bInSaveConfig=*/false);
 	Config.AddCustomTextFilter(InNewFilterData);
-	ToolSettings->SaveConfig();
 
 	AddCustomTextFilter(NewCustomTextFilter, false);
 
@@ -1645,9 +1617,8 @@ bool FNavigationToolFilterBar::TryDeleteCustomTextFilter(const FString& InFilter
 		return false;
 	}
 
-	FSequencerFilterBarConfig& Config = ToolSettings->FindOrAddFilterBar(GetIdentifier(), false);
+	FSequencerFilterBarConfig& Config = ToolSettings->FindOrAddFilterBar(GetIdentifier(), /*bInSaveConfig=*/false);
 	Config.RemoveCustomTextFilter(InFilterName);
-	ToolSettings->SaveConfig();
 
 	if (const TSharedPtr<FNavigationToolFilter_CustomText> CustomTextFilter = FindCustomTextFilterByDisplayName(InFilterName))
 	{
@@ -1682,7 +1653,7 @@ bool FNavigationToolFilterBar::TryModifyCustomTextFilter(const FCustomTextFilter
 	const bool bWasFilterEnabled = IsFilterEnabled(CustomTextFilterRef);
 	const bool bWasFilterActive = IsFilterActive(CustomTextFilterRef);
 
-	FSequencerFilterBarConfig& Config = ToolSettings->FindOrAddFilterBar(GetIdentifier(), false);
+	FSequencerFilterBarConfig& Config = ToolSettings->FindOrAddFilterBar(GetIdentifier(), /*bInSaveConfig=*/false);
 	Config.RemoveCustomTextFilter(InOldFilterName);
 
 	if (CustomTextFilter.IsValid())
@@ -1691,7 +1662,6 @@ bool FNavigationToolFilterBar::TryModifyCustomTextFilter(const FCustomTextFilter
 	}
 
 	Config.AddCustomTextFilter(InFilterData);
-	ToolSettings->SaveConfig();
 
 	RemoveCustomTextFilter(CustomTextFilterRef, false);
 	AddCustomTextFilter(CustomTextFilterRef, false);
@@ -1704,6 +1674,10 @@ bool FNavigationToolFilterBar::TryModifyCustomTextFilter(const FCustomTextFilter
 	{
 		SetFilterEnabledByDisplayName(NewFilterName, true, true);
 	}
+	else
+	{
+		ToolSettings->SaveConfig();
+	}
 
 	return true;
 }
diff --git a/Engine/Plugins/Experimental/SequenceNavigator/Source/SequenceNavigator/Private/NavigationToolExtender.cpp b/Engine/Plugins/Experimental/SequenceNavigator/Source/SequenceNavigator/Private/NavigationToolExtender.cpp
--- a/Engine/Plugins/Experimental/SequenceNavigator/Source/SequenceNavigator/Private/NavigationToolExtender.cpp
+++ b/Engine/Plugins/Experimental/SequenceNavigator/Source/SequenceNavigator/Private/NavigationToolExtender.cpp
@@ -96,6 +96,11 @@ void FNavigationToolExtender::OnSequencerCreated(const TSharedRef<ISequencer> In
 
 void FNavigationToolExtender::OnSequencerActivated(const FName InToolId, FMovieSceneSequenceIDRef InSequenceId)
 {
+	if (GIsTransacting)
+	{
+		return;
+	}
+
 	FNavigationToolInstance* const ToolInstance = ToolInstances.Find(InToolId);
 	if (!ToolInstance || !ToolInstance->WeakSequencer.IsValid())
 	{
