[SequenceNavigator] Remove default filter bar visibility toggle hotkey as it creates conflict with the Sequencer filter bar visibility toggle hotkey

diff --git a/Engine/Plugins/Experimental/SequenceNavigator/Source/SequenceNavigator/Private/Filters/NavigationToolFilterCommands.cpp b/Engine/Plugins/Experimental/SequenceNavigator/Source/SequenceNavigator/Private/Filters/NavigationToolFilterCommands.cpp
--- a/Engine/Plugins/Experimental/SequenceNavigator/Source/SequenceNavigator/Private/Filters/NavigationToolFilterCommands.cpp
+++ b/Engine/Plugins/Experimental/SequenceNavigator/Source/SequenceNavigator/Private/Filters/NavigationToolFilterCommands.cpp
@@ -10,7 +10,7 @@ void FNavigationToolFilterCommands::RegisterCommands()
 		, "Filter Bar Visibility"
 		, "Toggle the visibility of the filter bar. The filter bar will only be displayed when there are enabled filters"
 		, EUserInterfaceActionType::ToggleButton
-		, FInputChord(EModifierKey::Alt, EKeys::F));
+		, FInputChord());
 
 	UI_COMMAND(SetToVerticalLayout
 		, "Vertical"
