[SequenceNavigator]
- Check UToolMenus() module loaded and valid before add/remove extension
- Update logs

diff --git a/Engine/Plugins/Experimental/SequenceNavigator/Source/SequenceNavigator/Private/NavigationToolExtender.cpp b/Engine/Plugins/Experimental/SequenceNavigator/Source/SequenceNavigator/Private/NavigationToolExtender.cpp
--- a/Engine/Plugins/Experimental/SequenceNavigator/Source/SequenceNavigator/Private/NavigationToolExtender.cpp
+++ b/Engine/Plugins/Experimental/SequenceNavigator/Source/SequenceNavigator/Private/NavigationToolExtender.cpp
@@ -51,12 +51,9 @@ FNavigationToolExtender::~FNavigationToolExtender()
 		SequencerModule.UnregisterOnSequencerCreated(SequencerCreatedHandle);
 	}
 
-	if (FModuleManager::Get().IsModuleLoaded(TEXT("ToolMenus")))
+	for (const TPair<FName, FNavigationToolInstance>& ToolInstance : ToolInstances)
 	{
-		for (const TPair<FName, FNavigationToolInstance>& ToolInstance : ToolInstances)
-		{
-			RemoveSequencerToolBarExtension(ToolInstance.Value);
-		}
+		RemoveSequencerToolBarExtension(ToolInstance.Value);
 	}
 }
 
@@ -201,7 +198,8 @@ bool FNavigationToolExtender::RegisterToolProvider(const TSharedRef<ISequencer>&
 	const FName ToolId = GetToolInstanceId(*InSequencer);
 	if (ToolId.IsNone())
 	{
-		UE_LOG(LogSequenceNavigator, Error, TEXT("Invalid Sequence Navigator Instance Name"));
+		UE_LOG(LogSequenceNavigator, Error, TEXT("Invalid tool instance name: %s")
+			, *ToolId.ToString());
 		return false;
 	}
 
@@ -209,13 +207,14 @@ bool FNavigationToolExtender::RegisterToolProvider(const TSharedRef<ISequencer>&
 
 	if (NewProviderId.IsNone())
 	{
-		UE_LOG(LogSequenceNavigator, Error, TEXT("Invalid Sequence Navigator Provider Name"));
+		UE_LOG(LogSequenceNavigator, Error, TEXT("Invalid provider name: %s")
+			, *NewProviderId.ToString());
 		return false;
 	}
 
 	if (const TSharedPtr<FNavigationToolProvider> ExistingProvider = FindToolProvider(ToolId, NewProviderId))
 	{
-		UE_LOG(LogSequenceNavigator, Warning, TEXT("Sequence Navigator provider already registered: %s")
+		UE_LOG(LogSequenceNavigator, Warning, TEXT("Provider already registered: %s")
 			, *NewProviderId.ToString());
 		return true;
 	}
@@ -227,7 +226,7 @@ bool FNavigationToolExtender::RegisterToolProvider(const TSharedRef<ISequencer>&
 	FNavigationToolInstance* const ToolInstance = This.ToolInstances.Find(ToolId);
 	if (!ToolInstance)
 	{
-		UE_LOG(LogSequenceNavigator, Error, TEXT("Sequence Navigator instance could not be created for provider: %s")
+		UE_LOG(LogSequenceNavigator, Error, TEXT("Tool instance could not be created for provider: %s")
 			, *NewProviderId.ToString());
 		return true;
 	}
@@ -236,7 +235,7 @@ bool FNavigationToolExtender::RegisterToolProvider(const TSharedRef<ISequencer>&
 
 	AddSequencerToolBarExtension(*ToolInstance);
 
-	UE_LOG(LogSequenceNavigator, Log, TEXT("Sequence Navigator '%s' Instance Provider '%s' Registered")
+	UE_LOG(LogSequenceNavigator, Log, TEXT("Tool instance '%s' provider '%s' registered")
 		, *ToolId.ToString(), *NewProviderId.ToString());
 
 	This.ProvidersChangedDelegate.Broadcast(ToolId, InProvider, ENavigationToolProvidersChangeType::Add);
@@ -261,7 +260,7 @@ bool FNavigationToolExtender::UnregisterToolProvider(const FName InToolId, const
 		const TSharedRef<FNavigationToolProvider>& Provider = *It;
 		if (Provider->GetIdentifier() == InProviderId)
 		{
-			UE_LOG(LogSequenceNavigator, Log, TEXT("Sequence Navigator Provider Unregistered: %s")
+			UE_LOG(LogSequenceNavigator, Log, TEXT("Provider unregistered: %s")
 				, *Provider->GetIdentifier().ToString());
 
 			This.ProvidersChangedDelegate.Broadcast(InToolId, Provider, ENavigationToolProvidersChangeType::Remove);
@@ -361,7 +360,18 @@ void FNavigationToolExtender::AddSequencerToolBarExtension(const FNavigationTool
 		return;
 	}
 
-	UToolMenu* const ToolMenu = UToolMenus::Get()->ExtendMenu(TEXT("Sequencer.MainToolBar"));
+	if (!FModuleManager::Get().IsModuleLoaded(TEXT("ToolMenus")))
+	{
+		return;
+	}
+
+	UToolMenus* const ToolMenus = UToolMenus::Get();
+	if (!ToolMenus)
+	{
+		return;
+	}
+
+	UToolMenu* const ToolMenu = ToolMenus->ExtendMenu(TEXT("Sequencer.MainToolBar"));
 	if (!ToolMenu)
 	{
 		return;
@@ -380,7 +390,18 @@ void FNavigationToolExtender::AddSequencerToolBarExtension(const FNavigationTool
 
 void FNavigationToolExtender::RemoveSequencerToolBarExtension(const FNavigationToolInstance& InToolInstance)
 {
-	if (UToolMenu* const ToolMenu = UToolMenus::Get()->FindMenu(TEXT("Sequencer.MainToolBar")))
+	if (!FModuleManager::Get().IsModuleLoaded(TEXT("ToolMenus")))
+	{
+		return;
+	}
+
+	UToolMenus* const ToolMenus = UToolMenus::Get();
+	if (!ToolMenus)
+	{
+		return;
+	}
+
+	if (UToolMenu* const ToolMenu = ToolMenus->FindMenu(TEXT("Sequencer.MainToolBar")))
 	{
 		ToolMenu->RemoveSection(TEXT("NavigationTool"));
 	}
diff --git a/Engine/Plugins/Experimental/SequenceNavigator/Source/SequenceNavigator/Private/NavigationToolView.cpp b/Engine/Plugins/Experimental/SequenceNavigator/Source/SequenceNavigator/Private/NavigationToolView.cpp
--- a/Engine/Plugins/Experimental/SequenceNavigator/Source/SequenceNavigator/Private/NavigationToolView.cpp
+++ b/Engine/Plugins/Experimental/SequenceNavigator/Source/SequenceNavigator/Private/NavigationToolView.cpp
@@ -454,13 +454,13 @@ void FNavigationToolView::SaveViewItemFlags(const FNavigationToolItemPtr& InItem
 	FNavigationToolSaveState* const SaveState = Provider->GetSaveState(*Tool);
 	if (!SaveState)
 	{
-		UE_LOG(LogSequenceNavigator, Warning, TEXT("SaveViewItemFlags(): Sequence Navigator Save State is NULL!"));
+		UE_LOG(LogSequenceNavigator, Warning, TEXT("SaveViewItemFlags(): Save state is NULL!"));
 		return;
 	}
 
 	if (!SaveState->ToolViewSaveStates.IsValidIndex(ToolViewId))
 	{
-		UE_LOG(LogSequenceNavigator, Warning, TEXT("SaveViewItemFlags(): '%d' is in invalid saved tool view state index"), ToolViewId);
+		UE_LOG(LogSequenceNavigator, Warning, TEXT("SaveViewItemFlags(): Invalid tool view Id: %d"), ToolViewId);
 		return;
 	}
 
@@ -1710,6 +1710,10 @@ void FNavigationToolView::LoadViewState(const TSharedRef<FNavigationToolProvider
 		LoadFilterState(*ViewSaveState, false, false);
 		LoadToolViewItems(*ViewSaveState);
 	}
+	else
+	{
+		UE_LOG(LogSequenceNavigator, Warning, TEXT("FNavigationToolView::LoadViewState(): Save state is NULL!"));
+	}
 
 	PostLoad();
 
diff --git a/Engine/Plugins/Experimental/SequenceNavigator/Source/SequenceNavigator/Private/Providers/NavigationToolProvider.cpp b/Engine/Plugins/Experimental/SequenceNavigator/Source/SequenceNavigator/Private/Providers/NavigationToolProvider.cpp
--- a/Engine/Plugins/Experimental/SequenceNavigator/Source/SequenceNavigator/Private/Providers/NavigationToolProvider.cpp
+++ b/Engine/Plugins/Experimental/SequenceNavigator/Source/SequenceNavigator/Private/Providers/NavigationToolProvider.cpp
@@ -92,13 +92,13 @@ FNavigationToolViewSaveState* FNavigationToolProvider::GetViewSaveState(const IN
 	FNavigationToolSaveState* const SaveState = GetSaveState(InTool);
 	if (!SaveState)
 	{
-		UE_LOG(LogSequenceNavigator, Warning, TEXT("GetViewSaveState(): Sequence Navigator Save State is NULL!"));
+		UE_LOG(LogSequenceNavigator, Warning, TEXT("GetViewSaveState(): Save state is NULL!"));
 		return nullptr;
 	}
 
 	if (!SaveState->ToolViewSaveStates.IsValidIndex(InToolViewId))
 	{
-		UE_LOG(LogSequenceNavigator, Warning, TEXT("GetViewSaveState(): '%d' is in invalid saved tool view state index"), InToolViewId);
+		UE_LOG(LogSequenceNavigator, Warning, TEXT("GetViewSaveState(): Invalid tool view Id: %d"), InToolViewId);
 		return nullptr;
 	}
 
@@ -110,7 +110,7 @@ void FNavigationToolProvider::EnsureToolViewCount(const INavigationTool& InTool,
 	FNavigationToolSaveState* const SaveState = GetSaveState(InTool);
 	if (!SaveState)
 	{
-		UE_LOG(LogSequenceNavigator, Warning, TEXT("EnsureToolViewCount(): Sequence Navigator Save State is NULL!"));
+		UE_LOG(LogSequenceNavigator, Warning, TEXT("EnsureToolViewCount(): Save state is NULL!"));
 		return;
 	}
 
@@ -139,6 +139,9 @@ FReply FNavigationToolProvider::OnToolItemAcceptDrop(const FDragDropEvent& InDra
 
 void FNavigationToolProvider::Activate(FNavigationTool& InTool)
 {
+	const FName ProviderId = GetIdentifier();
+	UE_LOG(LogSequenceNavigator, Log, TEXT("Activating provider: %s"), *ProviderId.ToString());
+
 	if (const TSharedPtr<FUICommandList> BaseCommandList = InTool.GetBaseCommandList())
 	{
 		BindCommands(BaseCommandList.ToSharedRef());
@@ -159,6 +162,9 @@ void FNavigationToolProvider::Activate(FNavigationTool& InTool)
 
 void FNavigationToolProvider::Deactivate(FNavigationTool& InTool)
 {
+	const FName ProviderId = GetIdentifier();
+	UE_LOG(LogSequenceNavigator, Log, TEXT("Deactivating provider: %s"), *ProviderId.ToString());
+
 	CleanupExtendedColumnViews();
 
 	OnDeactivate();
@@ -193,6 +199,9 @@ bool FNavigationToolProvider::IsSequenceSupported(UMovieSceneSequence* const InS
 
 void FNavigationToolProvider::SaveState(FNavigationTool& InTool)
 {
+	const FName ProviderId = GetIdentifier();
+	UE_LOG(LogSequenceNavigator, Log, TEXT("Saving provider state: %s"), *ProviderId.ToString());
+
 	SaveSerializedTree(InTool, /*bInResetTree*/true);
 
 	const TSharedRef<FNavigationToolProvider> SharedThisRef = SharedThis(this);
@@ -221,6 +230,9 @@ void FNavigationToolProvider::SaveState(FNavigationTool& InTool)
 
 void FNavigationToolProvider::LoadState(FNavigationTool& InTool)
 {
+	const FName ProviderId = GetIdentifier();
+	UE_LOG(LogSequenceNavigator, Log, TEXT("Loading provider state: %s"), *ProviderId.ToString());
+
 	const UNavigationToolSettings* const ToolSettings = GetDefault<UNavigationToolSettings>();
 	if (!ToolSettings)
 	{
