[SequenceNavigator]
- Remove unneeded SaveAllProviderStates() function since item states are automatically updated and saved when an item state changes
- Fix regression causing item expansion states to not be loaded

diff --git a/Engine/Plugins/Experimental/SequenceNavigator/Source/SequenceNavigator/Private/NavigationTool.cpp b/Engine/Plugins/Experimental/SequenceNavigator/Source/SequenceNavigator/Private/NavigationTool.cpp
--- a/Engine/Plugins/Experimental/SequenceNavigator/Source/SequenceNavigator/Private/NavigationTool.cpp
+++ b/Engine/Plugins/Experimental/SequenceNavigator/Source/SequenceNavigator/Private/NavigationTool.cpp
@@ -107,8 +107,6 @@ void FNavigationTool::Shutdown()
 {
 	UnbindEvents();
 
-	SaveAllProviderStates();
-
 	ToolTab->OnVisibilityChanged().RemoveAll(this);
 	ToolTab->Shutdown();
 
@@ -148,15 +146,6 @@ void FNavigationTool::ForEachProvider(const TFunction<bool(const TSharedRef<FNav
 		});
 }
 
-void FNavigationTool::SaveAllProviderStates()
-{
-	ForEachProvider([this](const TSharedRef<FNavigationToolProvider>& InProvider)
-		{
-			InProvider->SaveState(*this);
-			return true;
-		});
-}
-
 USequencerSettings* FNavigationTool::GetSequencerSettings() const
 {
 	if (const TSharedPtr<ISequencer> Sequencer = WeakSequencer.Pin())
@@ -1060,8 +1049,6 @@ void FNavigationTool::ShowHideToolTab(const bool bInVisible)
 	else
 	{
 		UnbindEvents();
-
-		SaveAllProviderStates();
 	}
 }
 
diff --git a/Engine/Plugins/Experimental/SequenceNavigator/Source/SequenceNavigator/Private/NavigationTool.h b/Engine/Plugins/Experimental/SequenceNavigator/Source/SequenceNavigator/Private/NavigationTool.h
--- a/Engine/Plugins/Experimental/SequenceNavigator/Source/SequenceNavigator/Private/NavigationTool.h
+++ b/Engine/Plugins/Experimental/SequenceNavigator/Source/SequenceNavigator/Private/NavigationTool.h
@@ -44,8 +44,6 @@ public:
 	void Init();
 	void Shutdown();
 
-	void SaveAllProviderStates();
-
 	bool CanProcessSequenceSpawn(UMovieSceneSequence* const InSequence) const;
 
 	FOnToolLoaded OnToolLoaded;
