[SequenceNavigator]
- Make nav tool use its own commands instead of using Sequencers
- Move command registration from FNavigationToolExtender to module startup

diff --git a/Engine/Plugins/Experimental/SequenceNavigator/Source/SequenceNavigator/Private/NavigationTool.cpp b/Engine/Plugins/Experimental/SequenceNavigator/Source/SequenceNavigator/Private/NavigationTool.cpp
--- a/Engine/Plugins/Experimental/SequenceNavigator/Source/SequenceNavigator/Private/NavigationTool.cpp
+++ b/Engine/Plugins/Experimental/SequenceNavigator/Source/SequenceNavigator/Private/NavigationTool.cpp
@@ -43,6 +43,7 @@ FNavigationTool::FNavigationTool(const TWeakPtr<ISequencer>& InWeakSequencer)
 	: WeakSequencer(InWeakSequencer)
 	, ToolTab(MakeShared<FNavigationToolTab>(*this))
 	, RootItem(MakeShared<FNavigationToolTreeRoot>(*this))
+	, BaseCommandList(MakeShared<FUICommandList>())
 {
 }
 
@@ -162,7 +163,7 @@ bool FNavigationTool::CanProcessSequenceSpawn(UMovieSceneSequence* const InSeque
 
 TSharedPtr<FUICommandList> FNavigationTool::GetBaseCommandList() const
 {
-	return WeakBaseCommandList.Pin();
+	return BaseCommandList;
 }
 
 TArray<FName> FNavigationTool::GetRegisteredItemProxyTypeNames() const
@@ -253,11 +254,6 @@ TSharedPtr<ISequencer> FNavigationTool::GetSequencer() const
 	return WeakSequencer.Pin();
 }
 
-void FNavigationTool::SetBaseCommandList(const TSharedPtr<FUICommandList>& InBaseCommandList)
-{
-	WeakBaseCommandList = InBaseCommandList;
-}
-
 bool FNavigationTool::DoesGlobalFilterExist(const FName InFilterId)
 {
 	for (const TSharedPtr<FNavigationToolBuiltInFilter>& GlobalFilter : GlobalFilters)
diff --git a/Engine/Plugins/Experimental/SequenceNavigator/Source/SequenceNavigator/Private/NavigationTool.h b/Engine/Plugins/Experimental/SequenceNavigator/Source/SequenceNavigator/Private/NavigationTool.h
--- a/Engine/Plugins/Experimental/SequenceNavigator/Source/SequenceNavigator/Private/NavigationTool.h
+++ b/Engine/Plugins/Experimental/SequenceNavigator/Source/SequenceNavigator/Private/NavigationTool.h
@@ -73,8 +73,6 @@ public:
 
 	virtual TSharedPtr<ISequencer> GetSequencer() const override;
 
-	virtual void SetBaseCommandList(const TSharedPtr<FUICommandList>& InBaseCommandList) override;
-
 	virtual TSharedPtr<INavigationToolView> RegisterToolView(const int32 InToolViewId) override;
 	virtual TSharedPtr<INavigationToolView> GetToolView(const int32 InToolViewId) const override;
 	virtual TSharedPtr<INavigationToolView> GetMostRecentToolView() const override;
@@ -210,10 +208,9 @@ private:
 	/** The root of all the items in the outliner */
 	TSharedRef<FNavigationToolTreeRoot> RootItem;
 
-	/** the base command list gotten from the outliner provider to append our command list to */
-	TWeakPtr<FUICommandList> WeakBaseCommandList;
+	TSharedPtr<FUICommandList> BaseCommandList;
 
-	/** the map of the registered items */
+	/** The map of the registered items */
 	TMap<FNavigationToolItemId, FNavigationToolItemPtr> ItemMap;
 
 	TMap<FNavigationToolItemId, FNavigationToolItemPtr> ItemsPendingAdd;
diff --git a/Engine/Plugins/Experimental/SequenceNavigator/Source/SequenceNavigator/Private/NavigationToolExtender.cpp b/Engine/Plugins/Experimental/SequenceNavigator/Source/SequenceNavigator/Private/NavigationToolExtender.cpp
--- a/Engine/Plugins/Experimental/SequenceNavigator/Source/SequenceNavigator/Private/NavigationToolExtender.cpp
+++ b/Engine/Plugins/Experimental/SequenceNavigator/Source/SequenceNavigator/Private/NavigationToolExtender.cpp
@@ -25,8 +25,6 @@ FNavigationToolExtender& FNavigationToolExtender::Get()
 
 FNavigationToolExtender::FNavigationToolExtender()
 {
-	FNavigationToolCommands::Register();
-
 	IModularFeatures::Get().RegisterModularFeature(GetModularFeatureName(), this);
 
 	ISequencerModule& SequencerModule = FModuleManager::LoadModuleChecked<ISequencerModule>(TEXT("Sequencer"));
@@ -40,8 +38,6 @@ FNavigationToolExtender::FNavigationToolExtender()
 
 FNavigationToolExtender::~FNavigationToolExtender()
 {
-	FNavigationToolCommands::Unregister();
-
 	IModularFeatures::Get().UnregisterModularFeature(GetModularFeatureName(), this);
 
 	if (FModuleManager::Get().IsModuleLoaded(TEXT("Sequencer")))
@@ -85,7 +81,6 @@ void FNavigationToolExtender::OnSequencerCreated(const TSharedRef<ISequencer> In
 	if (!ToolInstance.Instance.IsValid())
 	{
 		const TSharedPtr<FNavigationTool> NewInstance = MakeShared<FNavigationTool>(InSequencer);
-		NewInstance->SetBaseCommandList(InSequencer->GetCommandBindings());
 		NewInstance->Init();
 		ToolInstance.Instance = NewInstance;
 	}
diff --git a/Engine/Plugins/Experimental/SequenceNavigator/Source/SequenceNavigator/Private/SequenceNavigatorModule.cpp b/Engine/Plugins/Experimental/SequenceNavigator/Source/SequenceNavigator/Private/SequenceNavigatorModule.cpp
--- a/Engine/Plugins/Experimental/SequenceNavigator/Source/SequenceNavigator/Private/SequenceNavigatorModule.cpp
+++ b/Engine/Plugins/Experimental/SequenceNavigator/Source/SequenceNavigator/Private/SequenceNavigatorModule.cpp
@@ -1,16 +1,23 @@
 ﻿// Copyright Epic Games, Inc. All Rights Reserved.
 
 #include "SequenceNavigatorModule.h"
+#include "NavigationToolCommands.h"
 #include "SequenceNavigatorLog.h"
 
 DEFINE_LOG_CATEGORY(LogSequenceNavigator);
 
 void FSequenceNavigatorModule::StartupModule()
 {
+	using namespace UE::SequenceNavigator;
+
+	FNavigationToolCommands::Register();
 }
 
 void FSequenceNavigatorModule::ShutdownModule()
 {
+	using namespace UE::SequenceNavigator;
+
+	FNavigationToolCommands::Unregister();
 }
 
 IMPLEMENT_MODULE(FSequenceNavigatorModule, SequenceNavigator)
diff --git a/Engine/Plugins/Experimental/SequenceNavigator/Source/SequenceNavigator/Public/INavigationTool.h b/Engine/Plugins/Experimental/SequenceNavigator/Source/SequenceNavigator/Public/INavigationTool.h
--- a/Engine/Plugins/Experimental/SequenceNavigator/Source/SequenceNavigator/Public/INavigationTool.h
+++ b/Engine/Plugins/Experimental/SequenceNavigator/Source/SequenceNavigator/Public/INavigationTool.h
@@ -54,9 +54,6 @@ public:
 	/** Gets the Command List that the Navigation Tool Views will use to append their Command Lists to */
 	virtual TSharedPtr<FUICommandList> GetBaseCommandList() const = 0;
 
-	/** Sets the Command List that the Navigation Tool Views will use to append their Command Lists to */
-	virtual void SetBaseCommandList(const TSharedPtr<FUICommandList>& InBaseCommandList) = 0;
-
 	/** Register a new Navigation Tool View to the Navigation Tool to the given id, replacing the old view that was bound to the given id */
 	virtual TSharedPtr<INavigationToolView> RegisterToolView(const int32 InToolViewId) = 0;
 
