[SequenceNavigator] Include optional param "bIncludeRoot" for GetParents() that defaults to including all root items except parent

diff --git a/Engine/Plugins/Experimental/SequenceNavigator/Source/SequenceNavigator/Public/Items/INavigationToolItem.h b/Engine/Plugins/Experimental/SequenceNavigator/Source/SequenceNavigator/Public/Items/INavigationToolItem.h
--- a/Engine/Plugins/Experimental/SequenceNavigator/Source/SequenceNavigator/Public/Items/INavigationToolItem.h
+++ b/Engine/Plugins/Experimental/SequenceNavigator/Source/SequenceNavigator/Public/Items/INavigationToolItem.h
@@ -227,12 +227,13 @@ public:
 	virtual FNavigationToolItemPtr GetParent() const = 0;
 
 	/** Gets the list of all parents up the tree from this Item */
-	virtual TSet<FNavigationToolItemPtr> GetParents() const
+	virtual TSet<FNavigationToolItemPtr> GetParents(const bool bInIncludeRoot = false) const
 	{
 		TSet<FNavigationToolItemPtr> OutParents;
 
 		FNavigationToolItemPtr CurrentParent = GetParent();
-		while (CurrentParent.IsValid())
+		while (CurrentParent.IsValid()
+			&& (bInIncludeRoot || CurrentParent->GetItemId() != FNavigationToolItemId::RootId))
 		{
 			OutParents.Add(CurrentParent);
 
