[SequenceNavigator] Fix slow focus root sequence from child sequence

diff --git a/Engine/Plugins/Experimental/SequenceNavigator/Source/SequenceNavigator/Private/NavigationTool.cpp b/Engine/Plugins/Experimental/SequenceNavigator/Source/SequenceNavigator/Private/NavigationTool.cpp
--- a/Engine/Plugins/Experimental/SequenceNavigator/Source/SequenceNavigator/Private/NavigationTool.cpp
+++ b/Engine/Plugins/Experimental/SequenceNavigator/Source/SequenceNavigator/Private/NavigationTool.cpp
@@ -1312,15 +1312,6 @@ void FNavigationTool::OnProvidersChanged(const FName InToolId
 		return;
 	}
 
-	if (InChangeType == ENavigationToolProvidersChangeType::Add)
-	{
-		InProvider->Activate(*this);
-	}
-	else if (InChangeType == ENavigationToolProvidersChangeType::Remove)
-	{
-		InProvider->Deactivate(*this);
-	}
-
 	RefreshGlobalFilters();
 
 	OnTreeViewChanged();
diff --git a/Engine/Plugins/Experimental/SequenceNavigator/Source/SequenceNavigator/Private/NavigationToolExtender.cpp b/Engine/Plugins/Experimental/SequenceNavigator/Source/SequenceNavigator/Private/NavigationToolExtender.cpp
--- a/Engine/Plugins/Experimental/SequenceNavigator/Source/SequenceNavigator/Private/NavigationToolExtender.cpp
+++ b/Engine/Plugins/Experimental/SequenceNavigator/Source/SequenceNavigator/Private/NavigationToolExtender.cpp
@@ -53,30 +53,37 @@ FNavigationToolExtender::~FNavigationToolExtender()
 	}
 }
 
-void FNavigationToolExtender::OnSequencerCreated(const TSharedRef<ISequencer> InSequencer)
+FNavigationToolExtender::FNavigationToolInstance* FNavigationToolExtender::FindOrAddToolInstance_Internal(
+	const TSharedRef<ISequencer>& InSequencer)
 {
-	const USequencerSettings* const SequencerSettings = InSequencer->GetSequencerSettings();
-	const FName SettingsName = SequencerSettings ? SequencerSettings->GetFName() : NAME_None;
-
-	if (SettingsName.IsNone())
+	const FName NewToolId = GetToolInstanceId(*InSequencer);
+	if (NewToolId.IsNone())
 	{
-		return;
+		return nullptr;
 	}
 
-	FNavigationToolInstance& ToolInstance = ToolInstances.FindOrAdd(SettingsName);
+	FNavigationToolInstance& ToolInstance = ToolInstances.FindOrAdd(NewToolId);
 
-	ToolInstance.ToolId = SettingsName;
+	ToolInstance.ToolId = NewToolId;
 	ToolInstance.WeakSequencer = InSequencer;
-	ToolInstance.ActivateSequenceHandle = InSequencer->OnActivateSequence().AddLambda(
-		[this, SettingsName](FMovieSceneSequenceIDRef InSequenceId)
-		{
-			OnSequencerActivated(SettingsName, InSequenceId);
-		});
-	ToolInstance.SequencerClosedHandle = InSequencer->OnCloseEvent().AddLambda(
-		[this, SettingsName](const TSharedRef<ISequencer> InSequencer)
-		{
-			OnSequencerClosed(SettingsName, InSequencer);
-		});
+
+	if (!ToolInstance.ActivateSequenceHandle.IsValid())
+	{
+		ToolInstance.ActivateSequenceHandle = InSequencer->OnActivateSequence().AddLambda(
+			[this, NewToolId](FMovieSceneSequenceIDRef InSequenceId)
+			{
+				OnSequencerActivated(NewToolId, InSequenceId);
+			});
+	}
+
+	if (!ToolInstance.SequencerClosedHandle.IsValid())
+	{
+		ToolInstance.SequencerClosedHandle = InSequencer->OnCloseEvent().AddLambda(
+			[this, NewToolId](const TSharedRef<ISequencer> InSequencer)
+			{
+				OnSequencerClosed(NewToolId, InSequencer);
+			});
+	}
 
 	if (!ToolInstance.Instance.IsValid())
 	{
@@ -84,39 +91,59 @@ void FNavigationToolExtender::OnSequencerCreated(const TSharedRef<ISequencer> In
 		NewInstance->Init();
 		ToolInstance.Instance = NewInstance;
 	}
+
+	return &ToolInstance;
 }
 
-void FNavigationToolExtender::OnSequencerActivated(const FName InToolId, FMovieSceneSequenceIDRef InSequenceId)
+void FNavigationToolExtender::OnSequencerCreated(const TSharedRef<ISequencer> InSequencer)
 {
-	if (GIsTransacting)
+	const FName NewToolId = GetToolInstanceId(*InSequencer);
+	if (NewToolId.IsNone())
 	{
 		return;
 	}
 
-	FNavigationToolInstance* const ToolInstance = ToolInstances.Find(InToolId);
-	if (!ToolInstance || !ToolInstance->WeakSequencer.IsValid())
+	if (!AnyProviderSupportsSequencer(NewToolId, *InSequencer))
 	{
 		return;
 	}
 
-	// Only care about root sequence changes
-	const TSharedPtr<ISequencer> Sequencer = ToolInstance->WeakSequencer.Pin();
-	if (!Sequencer.IsValid() || InSequenceId != Sequencer->GetRootTemplateID())
+	FNavigationToolInstance* const NewToolInstance = FindOrAddToolInstance_Internal(InSequencer);
+	if (!NewToolInstance || !NewToolInstance->Instance.IsValid())
 	{
 		return;
 	}
 
-	if (!AnyProviderSupportsSequencer(InToolId, *Sequencer))
+	FNavigationTool& ToolRef = *NewToolInstance->Instance;
+
+	for (const TSharedRef<FNavigationToolProvider>& Provider : NewToolInstance->Providers)
+	{
+		Provider->Activate(ToolRef);
+	}
+
+	AddSequencerToolBarExtension(*NewToolInstance);
+}
+
+void FNavigationToolExtender::OnSequencerActivated(const FName InToolId, FMovieSceneSequenceIDRef InSequenceId)
+{
+	// Nothing to do here. May still be used for something, otherwise remove in the future
+	/*if (GIsTransacting)
 	{
 		return;
 	}
 
-	for (const TSharedRef<FNavigationToolProvider>& Provider : ToolInstance->Providers)
+	FNavigationToolInstance* const ToolInstance = ToolInstances.Find(InToolId);
+	if (!ToolInstance || !ToolInstance->WeakSequencer.IsValid())
 	{
-		Provider->Activate(*ToolInstance->Instance);
+		return;
 	}
 
-	AddSequencerToolBarExtension(*ToolInstance);
+	// Only care about root sequence changes
+	const TSharedPtr<ISequencer> Sequencer = ToolInstance->WeakSequencer.Pin();
+	if (!Sequencer.IsValid() || InSequenceId != Sequencer->GetRootTemplateID())
+	{
+		return;
+	}*/
 }
 
 void FNavigationToolExtender::OnSequencerClosed(const FName InToolId, const TSharedRef<ISequencer> InSequencer)
@@ -175,10 +202,13 @@ FName FNavigationToolExtender::GetToolInstanceId(const INavigationTool& InTool)
 
 TSharedPtr<INavigationTool> FNavigationToolExtender::FindNavigationTool(const TSharedRef<ISequencer>& InSequencer)
 {
-	const USequencerSettings* const SequencerSettings = InSequencer->GetSequencerSettings();
-	const FName SettingsName = SequencerSettings ? SequencerSettings->GetFName() : NAME_None;
+	const FName ToolId = GetToolInstanceId(*InSequencer);
+	if (ToolId.IsNone())
+	{
+		return nullptr;
+	}
 
-	FNavigationToolInstance* const ToolInstance = Get().ToolInstances.Find(SettingsName);
+	FNavigationToolInstance* const ToolInstance = Get().ToolInstances.Find(ToolId);
 	if (!ToolInstance)
 	{
 		return nullptr;
@@ -191,6 +221,7 @@ bool FNavigationToolExtender::RegisterToolProvider(const TSharedRef<ISequencer>&
 	, const TSharedRef<FNavigationToolProvider>& InProvider)
 {
 	const FName ToolId = GetToolInstanceId(*InSequencer);
+
 	if (ToolId.IsNone())
 	{
 		UE_LOG(LogSequenceNavigator, Error, TEXT("Invalid tool instance name: %s")
@@ -216,9 +247,7 @@ bool FNavigationToolExtender::RegisterToolProvider(const TSharedRef<ISequencer>&
 
 	FNavigationToolExtender& This = Get();
 
-	This.OnSequencerCreated(InSequencer);
-
-	FNavigationToolInstance* const ToolInstance = This.ToolInstances.Find(ToolId);
+	FNavigationToolInstance* const ToolInstance = This.FindOrAddToolInstance_Internal(InSequencer);
 	if (!ToolInstance)
 	{
 		UE_LOG(LogSequenceNavigator, Error, TEXT("Tool instance could not be created for provider: %s")
@@ -228,7 +257,12 @@ bool FNavigationToolExtender::RegisterToolProvider(const TSharedRef<ISequencer>&
 
 	ToolInstance->Providers.Add(InProvider);
 
-	AddSequencerToolBarExtension(*ToolInstance);
+	if (InProvider->IsSequenceSupported(InSequencer->GetRootMovieSceneSequence()))
+	{
+		InProvider->Activate(*ToolInstance->Instance);
+
+		AddSequencerToolBarExtension(*ToolInstance);
+	}
 
 	UE_LOG(LogSequenceNavigator, Log, TEXT("Tool instance '%s' provider '%s' registered")
 		, *ToolId.ToString(), *NewProviderId.ToString());
@@ -255,6 +289,17 @@ bool FNavigationToolExtender::UnregisterToolProvider(const FName InToolId, const
 		const TSharedRef<FNavigationToolProvider>& Provider = *It;
 		if (Provider->GetIdentifier() == InProviderId)
 		{
+			if (ToolInstance->Instance.IsValid())
+			{
+				if (const TSharedPtr<ISequencer> Sequencer = ToolInstance->WeakSequencer.Pin())
+				{
+					if (Provider->IsSequenceSupported(Sequencer->GetRootMovieSceneSequence()))
+					{
+						Provider->Deactivate(*ToolInstance->Instance);
+					}
+				}
+			}
+
 			UE_LOG(LogSequenceNavigator, Log, TEXT("Provider unregistered: %s")
 				, *Provider->GetIdentifier().ToString());
 
@@ -372,6 +417,8 @@ void FNavigationToolExtender::AddSequencerToolBarExtension(const FNavigationTool
 		return;
 	}
 
+	const FName ToolId = InToolInstance.ToolId;
+
 	FToolMenuSection& NavigationToolSection = ToolMenu->FindOrAddSection(TEXT("NavigationTool"));
 
 	FToolMenuEntry ToggleTabVisibleEntry = FToolMenuEntry::InitToolBarButton(
@@ -404,18 +451,18 @@ void FNavigationToolExtender::RemoveSequencerToolBarExtension(const FNavigationT
 
 bool FNavigationToolExtender::AnyProviderSupportsSequencer(const FName InToolId, const ISequencer& InSequencer)
 {
-	UMovieSceneSequence* const Sequence = InSequencer.GetRootMovieSceneSequence();
-	if (!Sequence)
+	UMovieSceneSequence* const RootSequence = InSequencer.GetRootMovieSceneSequence();
+	if (!RootSequence)
 	{
 		return false;
 	}
 
 	bool bSupportsSequencer = false;
 
-	ForEachToolProvider(InToolId, [Sequence, &bSupportsSequencer]
+	ForEachToolProvider(InToolId, [RootSequence, &bSupportsSequencer]
 		(const TSharedRef<FNavigationToolProvider>& InProvider)
 		{
-			if (InProvider->IsSequenceSupported(Sequence))
+			if (InProvider->IsSequenceSupported(RootSequence))
 			{
 				bSupportsSequencer = true;
 				return false;
diff --git a/Engine/Plugins/Experimental/SequenceNavigator/Source/SequenceNavigator/Public/NavigationToolExtender.h b/Engine/Plugins/Experimental/SequenceNavigator/Source/SequenceNavigator/Public/NavigationToolExtender.h
--- a/Engine/Plugins/Experimental/SequenceNavigator/Source/SequenceNavigator/Public/NavigationToolExtender.h
+++ b/Engine/Plugins/Experimental/SequenceNavigator/Source/SequenceNavigator/Public/NavigationToolExtender.h
@@ -135,8 +135,27 @@ public:
 	static bool AnyProviderSupportsSequencer(const FName InToolId, const ISequencer& InSequencer);
 
 protected:
+	struct FNavigationToolInstance
+	{
+		/** The Sequencer type name to register the provider for */
+		FName ToolId;
+
+		TSharedPtr<FNavigationTool> Instance;
+
+		/** The sequencer instance for this Navigation Tool instance. May be null if the Sequencer is not currently open */
+		TWeakPtr<ISequencer> WeakSequencer;
+
+		FDelegateHandle ActivateSequenceHandle;
+		FDelegateHandle SequencerClosedHandle;
+
+		/** List of providers registered to this tool instance */
+		TSet<TSharedRef<FNavigationToolProvider>> Providers;
+	};
+
 	static TSet<TSharedRef<FNavigationToolProvider>>* FindToolProviders(const FName InToolId);
 
+	FNavigationToolInstance* FindOrAddToolInstance_Internal(const TSharedRef<ISequencer>& InSequencer);
+
 	void OnSequencerCreated(const TSharedRef<ISequencer> InSequencer);
 	void OnSequencerActivated(const FName InToolId, FMovieSceneSequenceIDRef InSequenceId);
 	void OnSequencerClosed(const FName InToolId, const TSharedRef<ISequencer> InSequencer);
@@ -153,23 +172,6 @@ protected:
 
 	FDelegateHandle SequencerCreatedHandle;
 
-	struct FNavigationToolInstance
-	{
-		/** The Sequencer type name to register the provider for */
-		FName ToolId;
-
-		TSharedPtr<FNavigationTool> Instance;
-
-		/** The sequencer instance for this Navigation Tool instance. May be null if the Sequencer is not currently open */
-		TWeakPtr<ISequencer> WeakSequencer;
-
-		FDelegateHandle ActivateSequenceHandle;
-		FDelegateHandle SequencerClosedHandle;
-
-		/** List of providers registered to this tool instance */
-		TSet<TSharedRef<FNavigationToolProvider>> Providers;
-	};
-
 	static void AddSequencerToolBarExtension(const FNavigationToolInstance& InToolInstance);
 	static void RemoveSequencerToolBarExtension(const FNavigationToolInstance& InToolInstance);
 
