[SequenceNavigator] Show the correct name for bound object items

diff --git a/Engine/Plugins/Experimental/SequenceNavigator/Source/SequenceNavigator/Private/Items/NavigationToolBinding.cpp b/Engine/Plugins/Experimental/SequenceNavigator/Source/SequenceNavigator/Private/Items/NavigationToolBinding.cpp
--- a/Engine/Plugins/Experimental/SequenceNavigator/Source/SequenceNavigator/Private/Items/NavigationToolBinding.cpp
+++ b/Engine/Plugins/Experimental/SequenceNavigator/Source/SequenceNavigator/Private/Items/NavigationToolBinding.cpp
@@ -6,11 +6,11 @@
 #include "Items/NavigationToolItemUtils.h"
 #include "Items/NavigationToolSequence.h"
 #include "MovieScene.h"
-#include "ScopedTransaction.h"
 #include "MVVM/ObjectBindingModelStorageExtension.h"
 #include "MVVM/ViewModels/SequencerEditorViewModel.h"
 #include "MVVM/ViewModels/TrackModel.h"
 #include "NavigationToolScopedSelection.h"
+#include "ScopedTransaction.h"
 #include "Sections/MovieSceneSubSection.h"
 #include "Styling/SlateIconFinder.h"
 #include "Styling/StyleColors.h"
@@ -61,7 +61,11 @@ void FNavigationToolBinding::FindChildren(TArray<FNavigationToolItemPtr>& OutChi
 
 FText FNavigationToolBinding::GetDisplayName() const
 {
-	return FText::FromString(Binding.GetName());
+	if (UMovieScene* const MovieScene = GetMovieScene())
+	{
+		return MovieScene->GetObjectDisplayName(Binding.GetObjectGuid());
+	}
+	return FText::GetEmpty();
 }
 
 FText FNavigationToolBinding::GetClassName() const
@@ -280,6 +284,15 @@ UMovieSceneSequence* FNavigationToolBinding::GetSequence() const
 	return nullptr;
 }
 
+UMovieScene* FNavigationToolBinding::GetMovieScene() const
+{
+	if (UMovieSceneSequence* const Sequence = GetSequence())
+	{
+		return Sequence->GetMovieScene();
+	}
+	return nullptr;
+}
+
 TViewModelPtr<FObjectBindingModel> FNavigationToolBinding::GetViewModel() const
 {
 	const FGuid ObjectGuid = Binding.GetObjectGuid();
diff --git a/Engine/Plugins/Experimental/SequenceNavigator/Source/SequenceNavigator/Public/Items/NavigationToolBinding.h b/Engine/Plugins/Experimental/SequenceNavigator/Source/SequenceNavigator/Public/Items/NavigationToolBinding.h
--- a/Engine/Plugins/Experimental/SequenceNavigator/Source/SequenceNavigator/Public/Items/NavigationToolBinding.h
+++ b/Engine/Plugins/Experimental/SequenceNavigator/Source/SequenceNavigator/Public/Items/NavigationToolBinding.h
@@ -92,6 +92,9 @@ public:
 	/** @return The Sequence this binding belongs to */
 	UMovieSceneSequence* GetSequence() const;
 
+	/** @return The Movie Scene this binding belongs to */
+	UMovieScene* GetMovieScene() const;
+
 	/** @return The cached object that is bound in sequencer. */
 	UObject* GetCachedBoundObject() const;
 
