[SequenceNavigator] Disable item context menu "Duplicate" by default unless overriden

diff --git a/Engine/Plugins/Experimental/SequenceNavigator/Source/SequenceNavigator/Private/NavigationToolView.cpp b/Engine/Plugins/Experimental/SequenceNavigator/Source/SequenceNavigator/Private/NavigationToolView.cpp
--- a/Engine/Plugins/Experimental/SequenceNavigator/Source/SequenceNavigator/Private/NavigationToolView.cpp
+++ b/Engine/Plugins/Experimental/SequenceNavigator/Source/SequenceNavigator/Private/NavigationToolView.cpp
@@ -1411,21 +1411,23 @@ bool FNavigationToolView::CanDeleteSelected() const
 
 void FNavigationToolView::DuplicateSelected()
 {
-	if (const TSharedPtr<FNavigationTool> Tool = WeakTool.Pin())
+	// Disabled for now until (if) we decide to support this
+	/*if (const TSharedPtr<FNavigationTool> Tool = WeakTool.Pin())
 	{
-		//Tool->DuplicateItems(GetSelectedItems(), nullptr, TOptional<EItemDropZone>());
-	}
+		Tool->DuplicateItems(GetSelectedItems(), nullptr, TOptional<EItemDropZone>());
+	}*/
 }
 
 bool FNavigationToolView::CanDuplicateSelected() const
 {
-	for (const FNavigationToolItemPtr& Item : GetSelectedItems())
+	// Disabled for now until (if) we decide to support this
+	/*for (const FNavigationToolItemPtr& Item : GetSelectedItems())
 	{
-		if (Item && Item->IsA<FNavigationToolActor>())
+		if (Item && Item->CanDuplicate())
 		{
 			return true;
 		}
-	}
+	}*/
 	return false;
 }
 
