[SequenceNavigator] Don't reset columns to default when creating a filter

diff --git a/Engine/Plugins/Experimental/SequenceNavigator/Source/SequenceNavigator/Private/Widgets/SNavigationToolView.cpp b/Engine/Plugins/Experimental/SequenceNavigator/Source/SequenceNavigator/Private/Widgets/SNavigationToolView.cpp
--- a/Engine/Plugins/Experimental/SequenceNavigator/Source/SequenceNavigator/Private/Widgets/SNavigationToolView.cpp
+++ b/Engine/Plugins/Experimental/SequenceNavigator/Source/SequenceNavigator/Private/Widgets/SNavigationToolView.cpp
@@ -122,7 +122,7 @@ void SNavigationToolView::Construct(const FArguments& InArgs, const TSharedRef<F
 		FilterBar->OnFiltersChanged().AddSP(this, &SNavigationToolView::OnFiltersChanged);
 	}
 
-	RebuildWidget();
+	RebuildWidget(true);
 }
 
 SNavigationToolView::~SNavigationToolView()
@@ -137,7 +137,7 @@ SNavigationToolView::~SNavigationToolView()
 	}
 }
 
-void SNavigationToolView::RebuildWidget()
+void SNavigationToolView::RebuildWidget(const bool bInReconstructColumns)
 {
 	bIsEnterLastKeyPressed = false;
 
@@ -155,7 +155,10 @@ void SNavigationToolView::RebuildWidget()
 		}
 	}
 
-	ReconstructColumns();
+	if (bInReconstructColumns)
+	{
+		ReconstructColumns();
+	}
 
 	ChildSlot
 	[
@@ -353,7 +356,7 @@ void SNavigationToolView::ReconstructColumns()
 
 	HeaderRowWidget->ClearColumns();
 
-	// Add the Columns allocated by the Owning Instance
+	// Add the columns allocated by the owning instance
 	for (const TPair<FName, TSharedPtr<INavigationToolColumn>>& Pair : ToolView->GetColumns())
 	{
 		OwnerTool->ForEachProvider([this, &ToolView, &OwnerTool, &Pair]
diff --git a/Engine/Plugins/Experimental/SequenceNavigator/Source/SequenceNavigator/Private/Widgets/SNavigationToolView.h b/Engine/Plugins/Experimental/SequenceNavigator/Source/SequenceNavigator/Private/Widgets/SNavigationToolView.h
--- a/Engine/Plugins/Experimental/SequenceNavigator/Source/SequenceNavigator/Private/Widgets/SNavigationToolView.h
+++ b/Engine/Plugins/Experimental/SequenceNavigator/Source/SequenceNavigator/Private/Widgets/SNavigationToolView.h
@@ -36,8 +36,8 @@ public:
 
 	TSharedPtr<SNavigationToolTreeView> GetTreeView() const { return TreeView; }
 
-	/** Rebuilds the entire widget. Called when the filter bar layout is changed. */
-	void RebuildWidget();
+	/** Rebuilds the entire widget. Needs to be called when the filter bar layout changes. */
+	void RebuildWidget(const bool bInReconstructColumns = false);
 
 	void ReconstructColumns();
 
