[SequenceNavigator] Update takes column widget correctly and fix edge cases where the take couldn't be changed

diff --git a/Engine/Plugins/Experimental/SequenceNavigator/Source/SequenceNavigator/Private/Widgets/Columns/SNavigationToolTake.cpp b/Engine/Plugins/Experimental/SequenceNavigator/Source/SequenceNavigator/Private/Widgets/Columns/SNavigationToolTake.cpp
--- a/Engine/Plugins/Experimental/SequenceNavigator/Source/SequenceNavigator/Private/Widgets/Columns/SNavigationToolTake.cpp
+++ b/Engine/Plugins/Experimental/SequenceNavigator/Source/SequenceNavigator/Private/Widgets/Columns/SNavigationToolTake.cpp
@@ -19,91 +19,12 @@
 namespace UE::SequenceNavigator
 {
 
-DECLARE_DELEGATE_RetVal_OneParam(FReply, FOnTakeEntrySelected, const TSharedPtr<SNavigationToolTake::FTakeItemInfo> /*InTakeInfo*/);
-
-class SNavigationToolTakeEntry : public SCompoundWidget
-{
-public:
-	SLATE_BEGIN_ARGS(SNavigationToolTakeEntry) {}
-		SLATE_EVENT(FOnTakeEntrySelected, OnEntrySelected)
-	SLATE_END_ARGS()
-
-	void Construct(const FArguments& InArgs, const TSharedPtr<SNavigationToolTake::FTakeItemInfo>& InTakeEntry)
-	{
-		TakeEntry  = InTakeEntry;
-		MenuButtonStyle = &FAppStyle::Get().GetWidgetStyle<FButtonStyle>("Menu.Button");
-
-		OnTakeEntrySelected = InArgs._OnEntrySelected;
-
-		ChildSlot
-		[
-			SNew(SBox)
-			.WidthOverride(120.f)
-			[
-				SNew(SBorder)
-				.BorderImage(this, &SNavigationToolTakeEntry::GetBorderImage)
-				.Padding(FMargin(12.f, 1.f, 12.f, 1.f))
-				[
-					SNew(STextBlock)
-					.Text(FText::FromString(TakeEntry->DisplayName))
-					.ColorAndOpacity(FStyleColors::Foreground)
-				]
-
-			]
-		];
-	}
-
-	const FSlateBrush* GetBorderImage() const
-	{
-		if (IsHovered())
-		{
-			return &MenuButtonStyle->Hovered;
-		}
-		return &MenuButtonStyle->Normal;
-	}
-
-	virtual FReply OnMouseButtonUp(const FGeometry& InGeometry, const FPointerEvent& InMouseEvent) override
-	{
-		if (OnTakeEntrySelected.IsBound())
-		{
-			return OnTakeEntrySelected.Execute(TakeEntry);
-		}
-		return FReply::Handled();
-	}
-
-protected:
-	TSharedPtr<SNavigationToolTake::FTakeItemInfo> TakeEntry;
-	FOnTakeEntrySelected OnTakeEntrySelected;
-	const FButtonStyle* MenuButtonStyle = nullptr;
-};
-
-class FTakeDragDropOp : public FDragDropOperation
-{
-public:
-	DRAG_DROP_OPERATOR_TYPE(FTakeDragDropOp, FDragDropOperation)
-
-	virtual TSharedPtr<SWidget> GetDefaultDecorator() const override
-	{
-		return SNullWidget::NullWidget;
-	}
-
-	static TSharedRef<FTakeDragDropOp> New(const TSharedPtr<SNavigationToolTake::FTakeItemInfo>& InTakeInfo)
-	{
-		const TSharedRef<FTakeDragDropOp> Operation = MakeShared<FTakeDragDropOp>();
-		Operation->TakeInfo = InTakeInfo;
-		Operation->Construct();
-		return Operation;
-	}
-
-	TSharedPtr<SNavigationToolTake::FTakeItemInfo> TakeInfo;
-};
-
 void SNavigationToolTake::Construct(const FArguments& InArgs
-	, const FNavigationToolItemRef& InItem
+	, const FNavigationToolItemWeakPtr& InWeakItem
 	, const TSharedRef<INavigationToolView>& InView
 	, const TSharedRef<SNavigationToolTreeRow>& InRowWidget)
 {
-	WeakItem = InItem;
+	WeakItem = InWeakItem;
 	WeakView = InView;
 	WeakRowWidget = InRowWidget;
 
@@ -118,51 +39,32 @@ void SNavigationToolTake::Construct(const FArguments& InArgs
 
 	ChildSlot
 	[
-		SNew(SComboBox<TSharedPtr<SNavigationToolTake::FTakeItemInfo>>)
+		SNew(SComboBox<TSharedPtr<FSequenceTakeEntry>>)
 		.OptionsSource(&CachedTakes)
 		.InitiallySelectedItem(ActiveTakeInfo)
 		.OnSelectionChanged(this, &SNavigationToolTake::OnSelectionChanged)
-		.OnGenerateWidget(this, &SNavigationToolTake::GenerateTakeWidget)
+		.OnGenerateWidget(this, &SNavigationToolTake::GenerateTakeWidget, /*bShowTakeIndex=*/false, /*bShowNumberedOf=*/true)
 		[
-			GenerateTakeWidget(ActiveTakeInfo)
+			GenerateTakeWidget(ActiveTakeInfo, /*bShowTakeIndex=*/false, /*bShowNumberedOf=*/false)
 		]
 	];
 }
 
-TSharedRef<SWidget> SNavigationToolTake::GenerateTakeWidget(const TSharedPtr<FTakeItemInfo> InTakeInfo)
-{
-	const FText TakeInfoText = FText::Format(LOCTEXT("TakeNumberLabel", "({0}/{1})")
-		, FText::FromString(FString::FromInt(InTakeInfo->TakeIndex + 1))
-		, FText::FromString(FString::FromInt(CachedTakes.Num())));
-
-	return SNew(SHorizontalBox)
-		+ SHorizontalBox::Slot()
-		.HAlign(HAlign_Left)
-		.VAlign(VAlign_Center)
-		.AutoWidth()
-		[
-			SNew(STextBlock)
-			.Text(FText::FromString(FString::FromInt(InTakeInfo->TakeNumber)))
-		]
-		+ SHorizontalBox::Slot()
-		.FillWidth(1.f)
-		.HAlign(HAlign_Right)
-		.VAlign(VAlign_Center)
-		.Padding(5.f, 0.f, 0.f, 0.f)
-		[
-			SNew(STextBlock)
-			.TextStyle(FAppStyle::Get(), TEXT("SmallText"))
-			.ColorAndOpacity(FStyleColors::Hover)
-			.Text(TakeInfoText)
-		];
-}
-
-void SNavigationToolTake::RemoveItemColor() const
+TSharedRef<SWidget> SNavigationToolTake::GenerateTakeWidget(const TSharedPtr<FSequenceTakeEntry> InTakeInfo
+	, const bool bShowTakeIndex, const bool bShowNumberedOf)
 {
-	if (const TSharedPtr<INavigationTool> Tool = WeakTool.Pin())
+	if (InTakeInfo.IsValid())
 	{
-		Tool->RemoveItemColor(WeakItem.Pin());
+		return SNew(SNavigationToolTakeEntry, InTakeInfo.ToSharedRef())
+			.TotalTakeCount_Lambda([this]()
+				{
+					return CachedTakes.Num();
+				})
+			.OnEntrySelected(this, &SNavigationToolTake::OnTakeEntrySelected)
+			.ShowTakeIndex(bShowTakeIndex)
+			.ShowNumberedOf(bShowNumberedOf);
 	}
+	return SNullWidget::NullWidget;
 }
 
 FSlateColor SNavigationToolTake::GetBorderColor() const
@@ -174,11 +76,11 @@ FSlateColor SNavigationToolTake::GetBorderColor() const
 	return FSlateColor::UseSubduedForeground();
 }
 
-FReply SNavigationToolTake::OnTakeEntrySelected(const TSharedPtr<FTakeItemInfo> InTakeInfo)
+FReply SNavigationToolTake::OnTakeEntrySelected(const TSharedRef<FSequenceTakeEntry> InTakeInfo)
 {
 	if (const TSharedPtr<INavigationTool> Tool = WeakTool.Pin())
 	{
-		SetActiveTake(InTakeInfo->WeakSequence.Get());
+		SetActiveTake(InTakeInfo);
 
 		return FReply::Handled();
 	}
@@ -186,87 +88,31 @@ FReply SNavigationToolTake::OnTakeEntrySelected(const TSharedPtr<FTakeItemInfo>
 	return FReply::Unhandled();
 }
 
-void SNavigationToolTake::Press()
-{
-	bIsPressed = true;
-}
-
-void SNavigationToolTake::Release()
-{
-	bIsPressed = false;
-}
-
-FReply SNavigationToolTake::OnMouseButtonDown(const FGeometry& MyGeometry, const FPointerEvent& MouseEvent)
-{
-	if (MouseEvent.IsMouseButtonDown(EKeys::LeftMouseButton))
-	{
-		Press();
-
-		const TSharedRef<SNavigationToolTake> This = SharedThis(this);
-		return FReply::Handled().CaptureMouse(This).DetectDrag(This, EKeys::LeftMouseButton);
-	}
-
-	return SCompoundWidget::OnMouseButtonDown(MyGeometry, MouseEvent);
-}
-
-FReply SNavigationToolTake::OnMouseButtonUp(const FGeometry& MyGeometry, const FPointerEvent& MouseEvent)
+void SNavigationToolTake::OnSelectionChanged(const TSharedPtr<FSequenceTakeEntry> InTakeInfo, const ESelectInfo::Type InSelectType)
 {
-	FReply Reply = FReply::Unhandled();
-
-	if (bIsPressed && IsHovered() && MouseEvent.GetEffectingButton() == EKeys::LeftMouseButton)
+	if (InTakeInfo.IsValid())
 	{
-		Release();
-		Reply = FReply::Handled();
+		SetActiveTake(InTakeInfo);
 	}
-
-	if (Reply.GetMouseCaptor().IsValid() == false && HasMouseCapture())
-	{
-		Reply.ReleaseMouseCapture();
-	}
-
-	return Reply;
 }
 
-void SNavigationToolTake::OnFocusLost(const FFocusEvent& InFocusEvent)
-{
-	Release();
-
-	SCompoundWidget::OnFocusLost(InFocusEvent);
-}
-
-FReply SNavigationToolTake::OnDragDetected(const FGeometry& MyGeometry, const FPointerEvent& MouseEvent)
+void SNavigationToolTake::CacheTakes()
 {
-	if (MouseEvent.IsMouseButtonDown(EKeys::LeftMouseButton))
-	{
-		return FReply::Handled().BeginDragDrop(FTakeDragDropOp::New(ActiveTakeInfo));
-	}
-
-	return SCompoundWidget::OnDragDetected(MyGeometry, MouseEvent);
-}
+	CachedTakes.Reset();
 
-void SNavigationToolTake::OnDragEnter(const FGeometry& MyGeometry, const FDragDropEvent& DragDropEvent)
-{
-	if (const TSharedPtr<FTakeDragDropOp> DragOp = DragDropEvent.GetOperationAs<FTakeDragDropOp>())
+	const TSharedPtr<INavigationToolItem> Item = WeakItem.Pin();
+	if (!Item.IsValid())
 	{
-		OnTakeEntrySelected(DragOp->TakeInfo);
+		return;
 	}
 
-	SCompoundWidget::OnDragEnter(MyGeometry, DragDropEvent);
-}
-
-void SNavigationToolTake::OnSelectionChanged(const TSharedPtr<FTakeItemInfo> InTakeInfo, const ESelectInfo::Type InSelectType)
-{
-	if (InTakeInfo.IsValid() && InTakeInfo->WeakSequence.IsValid())
+	FNavigationToolSequence* const SequenceItem = Item->CastTo<FNavigationToolSequence>();
+	if (!SequenceItem)
 	{
-		SetActiveTake(InTakeInfo->WeakSequence.Get());
+		return;
 	}
-}
-
-void SNavigationToolTake::CacheTakes()
-{
-	CachedTakes.Reset();
 
-	UMovieSceneSubSection* const SubSection = Cast<UMovieSceneSubSection>(GetSequenceItemSubSection(WeakItem.Pin()));
+	UMovieSceneSubSection* const SubSection = SequenceItem->GetSubSection();
 	if (!SubSection)
 	{
 		return;
@@ -296,10 +142,10 @@ void SNavigationToolTake::CacheTakes()
 		{
 			if (UMovieSceneSequence* const Sequence = Cast<UMovieSceneSequence>(ThisAssetData.GetAsset()))
 			{
-				const TSharedRef<FTakeItemInfo> NewTakeInfo = MakeShared<FTakeItemInfo>();
+				const TSharedRef<FSequenceTakeEntry> NewTakeInfo = MakeShared<FSequenceTakeEntry>();
 				NewTakeInfo->TakeIndex = TakeIndex;
 				NewTakeInfo->TakeNumber = TakeNumber;
-				NewTakeInfo->DisplayName = Sequence->GetDisplayName().ToString();
+				NewTakeInfo->DisplayName = Sequence->GetDisplayName();
 				NewTakeInfo->WeakSequence = Sequence;
 				CachedTakes.Add(NewTakeInfo);
 
@@ -314,8 +160,14 @@ void SNavigationToolTake::CacheTakes()
 	}
 }
 
-void SNavigationToolTake::SetActiveTake(UMovieSceneSequence* const InSequence)
+void SNavigationToolTake::SetActiveTake(const TSharedPtr<FSequenceTakeEntry>& InTakeInfo)
 {
+	UMovieSceneSequence* const Sequence = InTakeInfo->WeakSequence.Get();
+	if (!Sequence)
+	{
+		return;
+	}
+
 	const TSharedPtr<INavigationTool> Tool = WeakTool.Pin();
 	if (!Tool.IsValid())
 	{
@@ -328,62 +180,89 @@ void SNavigationToolTake::SetActiveTake(UMovieSceneSequence* const InSequence)
 		return;
 	}
 
-	bool bChangedTake = false;
-
-	const FScopedTransaction Transaction(LOCTEXT("ChangeTake_Transaction", "Change Take"));
-
-	TArray<UMovieSceneSection*> Sections;
-	Sequencer->GetSelectedSections(Sections);
+	const TSharedPtr<INavigationToolItem> Item = WeakItem.Pin();
+	if (!Item.IsValid())
+	{
+		return;
+	}
 
-	for (int32 SectionIndex = 0; SectionIndex < Sections.Num(); ++SectionIndex)
+	FNavigationToolSequence* const SequenceItem = Item->CastTo<FNavigationToolSequence>();
+	if (!SequenceItem)
 	{
-		if (!Sections[SectionIndex]->IsA<UMovieSceneSubSection>())
-		{
-			continue;
-		}
+		return;
+	}
 
-		UMovieSceneSubSection* Section = Cast<UMovieSceneSubSection>(Sections[SectionIndex]);
-		UMovieSceneSubTrack* SubTrack = CastChecked<UMovieSceneSubTrack>(Section->GetOuter());
+	const UMovieScene* const SequenceMovieScene = SequenceItem->GetSequenceMovieScene();
+	if (!SequenceMovieScene)
+	{
+		return;
+	}
 
-		TRange<FFrameNumber> NewSectionRange = Section->GetRange();
-		FFrameNumber		 NewSectionStartOffset = Section->Parameters.StartFrameOffset;
-		int32                NewSectionPrerollFrames = Section->GetPreRollFrames();
-		int32                NewRowIndex = Section->GetRowIndex();
-		FFrameNumber         NewSectionStartTime = NewSectionRange.GetLowerBound().IsClosed() ? UE::MovieScene::DiscreteInclusiveLower(NewSectionRange) : 0;
-		int32                NewSectionRowIndex = Section->GetRowIndex();
-		FColor               NewSectionColorTint = Section->GetColorTint();
+	UMovieSceneSubSection* const ParentSubSection = SequenceItem->GetSubSection();
+	if (!ParentSubSection)
+	{
+		return;
+	}
 
-		const int32 Duration = (NewSectionRange.GetLowerBound().IsClosed() && NewSectionRange.GetUpperBound().IsClosed()) ? UE::MovieScene::DiscreteSize(NewSectionRange) : 1;
-		UMovieSceneSubSection* NewSection = SubTrack->AddSequence(InSequence, NewSectionStartTime, Duration);
+	UMovieSceneSubTrack* const ParentSubTrack = Cast<UMovieSceneSubTrack>(ParentSubSection->GetOuter());
+	if (!ParentSubTrack)
+	{
+		return;
+	}
 
-		if (NewSection != nullptr)
-		{
-			SubTrack->RemoveSection(*Section);
+	bool bChangedTake = false;
 
-			NewSection->SetRange(NewSectionRange);
-			NewSection->Parameters.StartFrameOffset = NewSectionStartOffset;
-			NewSection->Parameters.TimeScale = Section->Parameters.TimeScale.DeepCopy(NewSection);
-			NewSection->SetPreRollFrames(NewSectionPrerollFrames);
-			NewSection->SetRowIndex(NewSectionRowIndex);
-			NewSection->SetColorTint(NewSectionColorTint);
+	FScopedTransaction Transaction(LOCTEXT("ChangeTake_Transaction", "Change Take"));
 
-			UMovieSceneCinematicShotSection* ShotSection = Cast<UMovieSceneCinematicShotSection>(Section);
-			UMovieSceneCinematicShotSection* NewShotSection = Cast<UMovieSceneCinematicShotSection>(NewSection);
+	const TRange<FFrameNumber> NewSectionRange = ParentSubSection->GetRange();
+	const FFrameNumber NewSectionStartOffset = ParentSubSection->Parameters.StartFrameOffset;
+	const int32 NewSectionPrerollFrames = ParentSubSection->GetPreRollFrames();
+	const int32 NewRowIndex = ParentSubSection->GetRowIndex();
+	const FFrameNumber NewSectionStartTime = NewSectionRange.GetLowerBound().IsClosed()
+		? MovieScene::DiscreteInclusiveLower(NewSectionRange) : 0;
+	const int32 NewSectionRowIndex = ParentSubSection->GetRowIndex();
+	const FColor NewSectionColorTint = ParentSubSection->GetColorTint();
 
-			// If the old shot's name is not the same as the sequence's name, assume the user had customized the shot name, so carry it over
-			if (ShotSection && NewShotSection && ShotSection->GetSequence() && ShotSection->GetShotDisplayName() != ShotSection->GetSequence()->GetName())
-			{
-				NewShotSection->SetShotDisplayName(ShotSection->GetShotDisplayName());
-			}
+	const int32 Duration = (NewSectionRange.GetLowerBound().IsClosed() && NewSectionRange.GetUpperBound().IsClosed())
+		? MovieScene::DiscreteSize(NewSectionRange) : 1;
 
-			bChangedTake = true;
+	if (UMovieSceneSubSection* const NewSection = ParentSubTrack->AddSequence(Sequence, NewSectionStartTime, Duration))
+	{
+		ParentSubTrack->RemoveSection(*ParentSubSection);
+
+		NewSection->SetRange(NewSectionRange);
+		NewSection->Parameters.StartFrameOffset = NewSectionStartOffset;
+		NewSection->Parameters.TimeScale = ParentSubSection->Parameters.TimeScale.DeepCopy(NewSection);
+		NewSection->SetPreRollFrames(NewSectionPrerollFrames);
+		NewSection->SetRowIndex(NewSectionRowIndex);
+		NewSection->SetColorTint(NewSectionColorTint);
+
+		// If the old shot's name is not the same as the sequence's name, assume the user had customized the shot name, so carry it over
+		UMovieSceneCinematicShotSection* const ParentShotSection = Cast<UMovieSceneCinematicShotSection>(ParentSubSection);
+		UMovieSceneCinematicShotSection* const NewShotSection = Cast<UMovieSceneCinematicShotSection>(NewSection);
+		if (NewShotSection && ParentShotSection)
+		{
+			if (const UMovieSceneSequence* const ParentShotSequence = ParentShotSection->GetSequence())
+            {
+                const FString ShowDisplayName = ParentShotSection->GetShotDisplayName();
+				if (ShowDisplayName != ParentShotSequence->GetName())
+				{
+					NewShotSection->SetShotDisplayName(ShowDisplayName);
+				}
+            }
 		}
+
+		bChangedTake = true;
 	}
 
 	if (bChangedTake)
 	{
 		Sequencer->NotifyMovieSceneDataChanged(EMovieSceneDataChangeType::MovieSceneStructureItemsChanged);
 	}
+	else
+	{
+		Transaction.Cancel();
+	}
 }
 
 } // namespace UE::SequenceNavigator
diff --git a/Engine/Plugins/Experimental/SequenceNavigator/Source/SequenceNavigator/Private/Widgets/Columns/SNavigationToolTake.h b/Engine/Plugins/Experimental/SequenceNavigator/Source/SequenceNavigator/Private/Widgets/Columns/SNavigationToolTake.h
--- a/Engine/Plugins/Experimental/SequenceNavigator/Source/SequenceNavigator/Private/Widgets/Columns/SNavigationToolTake.h
+++ b/Engine/Plugins/Experimental/SequenceNavigator/Source/SequenceNavigator/Private/Widgets/Columns/SNavigationToolTake.h
@@ -5,6 +5,7 @@
 #include "Columns/NavigationToolTakeColumn.h"
 #include "NavigationToolDefines.h"
 #include "ScopedTransaction.h"
+#include "SNavigationToolTakeEntry.h"
 #include "Templates/UniquePtr.h"
 #include "Widgets/SCompoundWidget.h"
 
@@ -17,48 +18,30 @@ class INavigationTool;
 class INavigationToolView;
 class SNavigationToolTreeRow;
 
+/** Column widget that displays a combo box dropdown list of takes for a specific sequence */
 class SNavigationToolTake : public SCompoundWidget
 {
 public:
-	struct FTakeItemInfo
-	{
-		uint32 TakeIndex = 0;
-		uint32 TakeNumber = 0;
-		FString DisplayName;
-		TWeakObjectPtr<UMovieSceneSequence> WeakSequence;
-	};
-
 	SLATE_BEGIN_ARGS(SNavigationToolTake) {}
 	SLATE_END_ARGS()
 
 	void Construct(const FArguments& InArgs
-		, const FNavigationToolItemRef& InItem
+		, const FNavigationToolItemWeakPtr& InWeakItem
 		, const TSharedRef<INavigationToolView>& InView
 		, const TSharedRef<SNavigationToolTreeRow>& InRowWidget);
 
-	void RemoveItemColor() const;
-
+private:
 	FSlateColor GetBorderColor() const;
 
-	FReply OnTakeEntrySelected(const TSharedPtr<FTakeItemInfo> InTakeInfo);
-
-	void Press();
-	void Release();
+	FReply OnTakeEntrySelected(const TSharedRef<FSequenceTakeEntry> InTakeInfo);
 
-	virtual FReply OnMouseButtonDown(const FGeometry& MyGeometry, const FPointerEvent& MouseEvent) override;
-	virtual FReply OnMouseButtonUp(const FGeometry& MyGeometry, const FPointerEvent& MouseEvent) override;
-	virtual void OnFocusLost(const FFocusEvent& InFocusEvent) override;
+	TSharedRef<SWidget> GenerateTakeWidget(const TSharedPtr<FSequenceTakeEntry> InTakeInfo
+		, const bool bShowTakeIndex, const bool bShowNumberedOf);
 
-	virtual FReply OnDragDetected(const FGeometry& MyGeometry, const FPointerEvent& MouseEvent) override;
-	virtual void OnDragEnter(const FGeometry& MyGeometry, const FDragDropEvent& DragDropEvent) override;
-
-private:
-	TSharedRef<SWidget> GenerateTakeWidget(const TSharedPtr<FTakeItemInfo> InTakeInfo);
-
-	void OnSelectionChanged(const TSharedPtr<FTakeItemInfo> InTakeInfo, const ESelectInfo::Type InSelectType);
+	void OnSelectionChanged(const TSharedPtr<FSequenceTakeEntry> InTakeInfo, const ESelectInfo::Type InSelectType);
 
 	void CacheTakes();
-	void SetActiveTake(UMovieSceneSequence* const InSequence);
+	void SetActiveTake(const TSharedPtr<FSequenceTakeEntry>& InTakeInfo);
 
 	FNavigationToolItemWeakPtr WeakItem;
 
@@ -68,11 +51,9 @@ private:
 
 	TWeakPtr<INavigationTool> WeakTool;
 
-	TArray<TSharedPtr<FTakeItemInfo>> CachedTakes;
-
-	TSharedPtr<FTakeItemInfo> ActiveTakeInfo;
+	TArray<TSharedPtr<FSequenceTakeEntry>> CachedTakes;
 
-	bool bIsPressed = false;
+	TSharedPtr<FSequenceTakeEntry> ActiveTakeInfo;
 
 	TUniquePtr<FScopedTransaction> UndoTransaction;
 };
diff --git a/Engine/Plugins/Experimental/SequenceNavigator/Source/SequenceNavigator/Private/Widgets/Columns/SNavigationToolTakeEntry.cpp b/Engine/Plugins/Experimental/SequenceNavigator/Source/SequenceNavigator/Private/Widgets/Columns/SNavigationToolTakeEntry.cpp
new file mode 100644
--- /dev/null
+++ b/Engine/Plugins/Experimental/SequenceNavigator/Source/SequenceNavigator/Private/Widgets/Columns/SNavigationToolTakeEntry.cpp
@@ -0,0 +1,234 @@
+﻿// Copyright Epic Games, Inc. All Rights Reserved.
+
+#include "SNavigationToolTakeEntry.h"
+#include "AssetRegistry/AssetData.h"
+#include "Framework/Application/SlateApplication.h"
+#include "Framework/MultiBox/MultiBoxBuilder.h"
+#include "Styling/StyleColors.h"
+#include "Widgets/Layout/SBorder.h"
+#include "Widgets/Layout/SBox.h"
+#include "Widgets/Text/STextBlock.h"
+#include "Widgets/SBoxPanel.h"
+
+#define LOCTEXT_NAMESPACE "SNavigationToolTakeEntry"
+
+namespace UE::SequenceNavigator
+{
+
+class FTakeDragDropOp : public FDragDropOperation
+{
+public:
+	DRAG_DROP_OPERATOR_TYPE(FTakeDragDropOp, FDragDropOperation)
+
+	virtual TSharedPtr<SWidget> GetDefaultDecorator() const override
+	{
+		return SNullWidget::NullWidget;
+	}
+
+	static TSharedRef<FTakeDragDropOp> New(const TSharedPtr<FSequenceTakeEntry>& InTakeInfo)
+	{
+		const TSharedRef<FTakeDragDropOp> Operation = MakeShared<FTakeDragDropOp>();
+		Operation->TakeInfo = InTakeInfo;
+		Operation->Construct();
+		return Operation;
+	}
+
+	TSharedPtr<FSequenceTakeEntry> TakeInfo;
+};
+
+void SNavigationToolTakeEntry::Construct(const FArguments& InArgs, const TSharedRef<FSequenceTakeEntry>& InTakeEntry)
+{
+	TakeEntry = InTakeEntry;
+
+	TotalTakeCount = InArgs._TotalTakeCount;
+	OnTakeEntrySelected = InArgs._OnEntrySelected;
+
+	ShowTakeIndex = InArgs._ShowTakeIndex;
+	ShowNumberedOf = InArgs._ShowNumberedOf;
+
+	const TSharedRef<SHorizontalBox> HorizontalBox = SNew(SHorizontalBox);
+
+	if (ShowTakeIndex)
+	{
+		HorizontalBox->AddSlot()
+			.HAlign(HAlign_Right)
+			.VAlign(VAlign_Center)
+			.AutoWidth()
+			[
+				SNew(STextBlock)
+				.TextStyle(FAppStyle::Get(), TEXT("SmallText"))
+				.ColorAndOpacity(FStyleColors::Hover)
+				.Text(this, &SNavigationToolTakeEntry::GetTakeNumberText)
+			];
+	}
+
+	HorizontalBox->AddSlot()
+		.FillWidth(1.f)
+		.HAlign(HAlign_Left)
+		.VAlign(VAlign_Center)
+		.Padding(3.f, 0.f, 0.f, 0.f)
+		[
+			SNew(STextBlock)
+			.Text(this, &SNavigationToolTakeEntry::GetTakeNameText)
+		];
+
+	if (ShowNumberedOf)
+	{
+		HorizontalBox->AddSlot()
+			.HAlign(HAlign_Right)
+			.VAlign(VAlign_Center)
+			.AutoWidth()
+			.Padding(3.f, 0.f, 0.f, 0.f)
+			[
+				SNew(STextBlock)
+				.TextStyle(FAppStyle::Get(), TEXT("SmallText"))
+				.ColorAndOpacity(FStyleColors::Hover)
+				.Text(this, &SNavigationToolTakeEntry::GetTakeIndexText)
+			];
+	}
+
+	ChildSlot
+	[
+		SNew(SBorder)
+		.BorderImage(this, &SNavigationToolTakeEntry::GetBorderImage)
+		.Padding(FMargin(2.f, 1.f, 3.f, 1.f))
+		[
+			HorizontalBox
+		]
+	];
+}
+
+const FSlateBrush* SNavigationToolTakeEntry::GetBorderImage() const
+{
+	return bHighlightForDragDrop ? FAppStyle::Get().GetBrush(TEXT("DetailsView.Highlight")) : nullptr;
+}
+
+FText SNavigationToolTakeEntry::GetTakeNumberText() const
+{
+	if (TakeEntry.IsValid())
+	{
+		return FText::Format(LOCTEXT("TakeNumberLabel", "{0}: ")
+			, FText::AsNumber(TakeEntry->TakeNumber));
+	}
+	return FText::GetEmpty();
+}
+
+FText SNavigationToolTakeEntry::GetTakeNameText() const
+{
+	if (TakeEntry.IsValid())
+	{
+		return TakeEntry->DisplayName;
+	}
+	return FText::GetEmpty();
+}
+
+FText SNavigationToolTakeEntry::GetTakeIndexText() const
+{
+	if (TakeEntry.IsValid() && TotalTakeCount.IsBound())
+	{
+		return FText::Format(LOCTEXT("TakeIndexLabel", "({0}/{1})")
+			, FText::AsNumber(TakeEntry->TakeIndex + 1)
+			, FText::AsNumber(TotalTakeCount.Get()));
+	}
+	return FText::GetEmpty();
+}
+
+FReply SNavigationToolTakeEntry::OnMouseButtonDown(const FGeometry& InGeometry, const FPointerEvent& InPointerEvent)
+{
+	FReply Reply = SCompoundWidget::OnMouseButtonDown(InGeometry, InPointerEvent);
+
+	if (InPointerEvent.IsMouseButtonDown(EKeys::LeftMouseButton))
+	{
+		const TSharedRef<SNavigationToolTakeEntry> This = SharedThis(this);
+		Reply = Reply.CaptureMouse(This).DetectDrag(This, EKeys::LeftMouseButton);
+	}
+
+	return Reply;
+}
+
+FReply SNavigationToolTakeEntry::OnMouseButtonUp(const FGeometry& InGeometry, const FPointerEvent& InPointerEvent)
+{
+	FReply Reply = FReply::Unhandled();
+
+	bHighlightForDragDrop = false;
+
+	if (IsHovered() && InPointerEvent.GetEffectingButton() == EKeys::LeftMouseButton)
+	{
+		Reply = FReply::Handled();
+	}
+
+	if (!Reply.GetMouseCaptor().IsValid() && HasMouseCapture())
+	{
+		Reply.ReleaseMouseCapture();
+	}
+
+	return Reply;
+}
+
+void SNavigationToolTakeEntry::OnFocusLost(const FFocusEvent& InFocusEvent)
+{
+	bHighlightForDragDrop = false;
+
+	SCompoundWidget::OnFocusLost(InFocusEvent);
+}
+
+FReply SNavigationToolTakeEntry::OnDragDetected(const FGeometry& InGeometry, const FPointerEvent& InPointerEvent)
+{
+	if (InPointerEvent.IsMouseButtonDown(EKeys::LeftMouseButton))
+	{
+		FSlateApplication::Get().DismissAllMenus();
+
+		return FReply::Handled().BeginDragDrop(FTakeDragDropOp::New(TakeEntry));
+	}
+
+	return SCompoundWidget::OnDragDetected(InGeometry, InPointerEvent);
+}
+
+void SNavigationToolTakeEntry::OnDragEnter(const FGeometry& InGeometry, const FDragDropEvent& InDragDropEvent)
+{
+	if (const TSharedPtr<FTakeDragDropOp> DragOp = InDragDropEvent.GetOperationAs<FTakeDragDropOp>())
+	{
+		bHighlightForDragDrop = true;
+	}
+
+	SCompoundWidget::OnDragEnter(InGeometry, InDragDropEvent);
+}
+
+void SNavigationToolTakeEntry::OnDragLeave(const FDragDropEvent& InDragDropEvent)
+{
+	if (const TSharedPtr<FTakeDragDropOp> DragOp = InDragDropEvent.GetOperationAs<FTakeDragDropOp>())
+	{
+		bHighlightForDragDrop = false;
+	}
+
+	SCompoundWidget::OnDragLeave(InDragDropEvent);
+}
+
+FReply SNavigationToolTakeEntry::OnDragOver(const FGeometry& InGeometry, const FDragDropEvent& InDragDropEvent)
+{
+	if (const TSharedPtr<FTakeDragDropOp> DragOp = InDragDropEvent.GetOperationAs<FTakeDragDropOp>())
+	{
+		bHighlightForDragDrop = true;
+	}
+
+	return SCompoundWidget::OnDragOver(InGeometry, InDragDropEvent);
+}
+
+FReply SNavigationToolTakeEntry::OnDrop(const FGeometry& InGeometry, const FDragDropEvent& InDragDropEvent)
+{
+	if (const TSharedPtr<FTakeDragDropOp> DragOp = InDragDropEvent.GetOperationAs<FTakeDragDropOp>())
+	{
+		if (OnTakeEntrySelected.IsBound())
+		{
+			return OnTakeEntrySelected.Execute(DragOp->TakeInfo.ToSharedRef());
+		}
+	}
+
+	bHighlightForDragDrop = false;
+
+	return SCompoundWidget::OnDrop(InGeometry, InDragDropEvent);
+}
+
+} // namespace UE::SequenceNavigator
+
+#undef LOCTEXT_NAMESPACE
diff --git a/Engine/Plugins/Experimental/SequenceNavigator/Source/SequenceNavigator/Private/Widgets/Columns/SNavigationToolTakeEntry.h b/Engine/Plugins/Experimental/SequenceNavigator/Source/SequenceNavigator/Private/Widgets/Columns/SNavigationToolTakeEntry.h
new file mode 100644
--- /dev/null
+++ b/Engine/Plugins/Experimental/SequenceNavigator/Source/SequenceNavigator/Private/Widgets/Columns/SNavigationToolTakeEntry.h
@@ -0,0 +1,78 @@
+﻿// Copyright Epic Games, Inc. All Rights Reserved.
+
+#pragma once
+
+#include "Widgets/SCompoundWidget.h"
+
+class UMovieSceneSequence;
+
+namespace UE::SequenceNavigator
+{
+
+/** Stores information about a sequence take */
+struct FSequenceTakeEntry
+{
+	/** Sequence that the take belongs to */
+	TWeakObjectPtr<UMovieSceneSequence> WeakSequence;
+
+	/** Index that the take is serialized to that will change depending on added/removed entries */
+	uint32 TakeIndex = 0;
+
+	/** Take number that doesn't change after being serialized */
+	uint32 TakeNumber = 0;
+
+	FText DisplayName;
+};
+
+/** Widget that display information about a single sequence take */
+class SNavigationToolTakeEntry : public SCompoundWidget
+{
+public:
+	DECLARE_DELEGATE_RetVal_OneParam(FReply, FOnTakeEntrySelected, const TSharedRef<FSequenceTakeEntry> /*InTakeInfo*/);
+
+	SLATE_BEGIN_ARGS(SNavigationToolTakeEntry) {}
+		SLATE_ATTRIBUTE(int32, TotalTakeCount)
+		/** Event called when a new take is selected */
+		SLATE_EVENT(FOnTakeEntrySelected, OnEntrySelected)
+		/** Shows the actual take index in the array in front of the take name */
+		SLATE_ARGUMENT(bool, ShowTakeIndex)
+		/** Shows numbered "1 / 2" text */
+		SLATE_ARGUMENT(bool, ShowNumberedOf)
+	SLATE_END_ARGS()
+
+	void Construct(const FArguments& InArgs, const TSharedRef<FSequenceTakeEntry>& InTakeEntry);
+
+	//~ Begin SWidget
+
+	virtual FReply OnMouseButtonDown(const FGeometry& InGeometry, const FPointerEvent& InPointerEvent) override;
+	virtual FReply OnMouseButtonUp(const FGeometry& InGeometry, const FPointerEvent& InPointerEvent) override;
+	virtual void OnFocusLost(const FFocusEvent& InFocusEvent) override;
+
+	virtual FReply OnDragDetected(const FGeometry& InGeometry, const FPointerEvent& InPointerEvent) override;
+	virtual void OnDragEnter(const FGeometry& InGeometry, const FDragDropEvent& InDragDropEvent) override;
+	virtual void OnDragLeave(const FDragDropEvent& InDragDropEvent) override;
+	virtual FReply OnDragOver(const FGeometry& InGeometry, const FDragDropEvent& InDragDropEvent) override;
+	virtual FReply OnDrop(const FGeometry& InGeometry, const FDragDropEvent& InDragDropEvent) override;
+
+	//~ End SWidget
+
+protected:
+	const FSlateBrush* GetBorderImage() const;
+
+	FText GetTakeNumberText() const;
+	FText GetTakeNameText() const;
+	FText GetTakeIndexText() const;
+
+	TSharedPtr<FSequenceTakeEntry> TakeEntry;
+
+	TAttribute<int32> TotalTakeCount;
+	FOnTakeEntrySelected OnTakeEntrySelected;
+
+	bool ShowTakeIndex = false;
+	bool ShowNumberedOf = false;
+
+	/** Highlight effect to display which entry is dragged over */
+	bool bHighlightForDragDrop = false;
+};
+
+} // namespace UE::SequenceNavigator
