[SequenceNavigator]
- Fix illegible highlighting of item names when search text is entered
- Remove invalid binding item text red color and only use icon red color to display invalid bindings (this matches the Curve Editor, maybe done to deal with the same situation)

diff --git a/Engine/Plugins/Experimental/SequenceNavigator/Source/SequenceNavigator/Private/Items/NavigationToolBinding.cpp b/Engine/Plugins/Experimental/SequenceNavigator/Source/SequenceNavigator/Private/Items/NavigationToolBinding.cpp
--- a/Engine/Plugins/Experimental/SequenceNavigator/Source/SequenceNavigator/Private/Items/NavigationToolBinding.cpp
+++ b/Engine/Plugins/Experimental/SequenceNavigator/Source/SequenceNavigator/Private/Items/NavigationToolBinding.cpp
@@ -77,11 +77,6 @@ FText FNavigationToolBinding::GetClassName() const
 	return FText::GetEmpty();
 }
 
-FSlateColor FNavigationToolBinding::GetItemLabelColor() const
-{
-	return IconColor;
-}
-
 FSlateIcon FNavigationToolBinding::GetIcon() const
 {
 	return Icon;
diff --git a/Engine/Plugins/Experimental/SequenceNavigator/Source/SequenceNavigator/Private/Widgets/Columns/SNavigationToolLabelItem.cpp b/Engine/Plugins/Experimental/SequenceNavigator/Source/SequenceNavigator/Private/Widgets/Columns/SNavigationToolLabelItem.cpp
--- a/Engine/Plugins/Experimental/SequenceNavigator/Source/SequenceNavigator/Private/Widgets/Columns/SNavigationToolLabelItem.cpp
+++ b/Engine/Plugins/Experimental/SequenceNavigator/Source/SequenceNavigator/Private/Widgets/Columns/SNavigationToolLabelItem.cpp
@@ -72,8 +72,8 @@ void SNavigationToolLabelItem::Construct(const FArguments& InArgs
 					.IsEnabled(this, &SNavigationToolLabelItem::IsItemEnabled)
 					.Text(this, &SNavigationToolLabelItem::GetItemText)
 					.Style(GetTextBlockStyle())
+					.Font(FAppStyle::GetFontStyle(TEXT("Sequencer.AnimationOutliner.RegularFont")))
 					.HighlightText(InRowWidget->GetHighlightText())
-					.ColorAndOpacity(this, &SNavigationToolLabelItem::GetForegroundColor)
 					.OnTextCommitted(this, &SNavigationToolLabelItem::OnLabelTextCommitted)
 					.OnVerifyTextChanged(this, &SNavigationToolLabelItem::OnVerifyItemLabelChanged)
 					.OnEnterEditingMode(this, &SNavigationToolLabelItem::OnEnterEditingMode)
@@ -86,29 +86,6 @@ void SNavigationToolLabelItem::Construct(const FArguments& InArgs
 	];
 }
 
-FSlateColor SNavigationToolLabelItem::GetForegroundColor() const
-{
-	const TSharedPtr<INavigationToolView> View = WeakView.Pin();
-	if (!View.IsValid())
-	{
-		return SCompoundWidget::GetForegroundColor();
-	}
-
-	const TSharedPtr<ISequencer> Sequencer = View->GetSequencer();
-	if (!Sequencer.IsValid())
-	{
-		return SCompoundWidget::GetForegroundColor();
-	}
-
-	const TSharedPtr<INavigationToolItem> Item = WeakItem.Pin();
-	if (!Item.IsValid())
-	{
-		return SCompoundWidget::GetForegroundColor();
-	}
-
-	return Item->GetItemLabelColor();
-}
-
 bool SNavigationToolLabelItem::IsReadOnly() const
 {
 	const FNavigationToolItemPtr Item = WeakItem.Pin();
diff --git a/Engine/Plugins/Experimental/SequenceNavigator/Source/SequenceNavigator/Private/Widgets/Columns/SNavigationToolLabelItem.h b/Engine/Plugins/Experimental/SequenceNavigator/Source/SequenceNavigator/Private/Widgets/Columns/SNavigationToolLabelItem.h
--- a/Engine/Plugins/Experimental/SequenceNavigator/Source/SequenceNavigator/Private/Widgets/Columns/SNavigationToolLabelItem.h
+++ b/Engine/Plugins/Experimental/SequenceNavigator/Source/SequenceNavigator/Private/Widgets/Columns/SNavigationToolLabelItem.h
@@ -45,8 +45,6 @@ public:
 
 	virtual const FInlineEditableTextBlockStyle* GetTextBlockStyle() const;
 
-	virtual FSlateColor GetForegroundColor() const;
-
 protected:
 	FNavigationToolItemWeakPtr WeakItem;
 
diff --git a/Engine/Plugins/Experimental/SequenceNavigator/Source/SequenceNavigator/Public/Items/NavigationToolBinding.h b/Engine/Plugins/Experimental/SequenceNavigator/Source/SequenceNavigator/Public/Items/NavigationToolBinding.h
--- a/Engine/Plugins/Experimental/SequenceNavigator/Source/SequenceNavigator/Public/Items/NavigationToolBinding.h
+++ b/Engine/Plugins/Experimental/SequenceNavigator/Source/SequenceNavigator/Public/Items/NavigationToolBinding.h
@@ -58,7 +58,6 @@ public:
 
 	UE_API virtual FText GetDisplayName() const override;
 	UE_API virtual FText GetClassName() const override;
-	UE_API virtual FSlateColor GetItemLabelColor() const override;
 	UE_API virtual FSlateIcon GetIcon() const override;
 	UE_API virtual FText GetIconTooltipText() const override;
 	UE_API virtual FSlateColor GetIconColor() const override;
