Sequencer: Check index before dereferencing

#rnx

diff --git a/Engine/Plugins/Experimental/SequenceNavigator/Source/SequenceNavigator/Private/Items/NavigationToolMarker.cpp b/Engine/Plugins/Experimental/SequenceNavigator/Source/SequenceNavigator/Private/Items/NavigationToolMarker.cpp
--- a/Engine/Plugins/Experimental/SequenceNavigator/Source/SequenceNavigator/Private/Items/NavigationToolMarker.cpp
+++ b/Engine/Plugins/Experimental/SequenceNavigator/Source/SequenceNavigator/Private/Items/NavigationToolMarker.cpp
@@ -202,6 +202,11 @@ bool FNavigationToolMarker::Rename(const FString& InName)
 	}
 
 	const TArray<FMovieSceneMarkedFrame>& MarkedFrames = MovieScene->GetMarkedFrames();
+	if (!MarkedFrames.IsValidIndex(MarkedFrameIndex))
+	{
+		return false;
+	}
+
 	FMovieSceneMarkedFrame& MarkedFrame = const_cast<FMovieSceneMarkedFrame&>(MarkedFrames[MarkedFrameIndex]);
 
 	bool bRenamed = false;
@@ -239,6 +244,11 @@ FFrameNumber FNavigationToolMarker::GetInTime() const
 	}
 
 	const TArray<FMovieSceneMarkedFrame>& MarkedFrames = MovieScene->GetMarkedFrames();
+	if (!MarkedFrames.IsValidIndex(MarkedFrameIndex))
+	{
+		return 0;
+	}
+
 	FMovieSceneMarkedFrame& MarkedFrame = const_cast<FMovieSceneMarkedFrame&>(MarkedFrames[MarkedFrameIndex]);
 
 	return MarkedFrame.FrameNumber;
@@ -264,6 +274,11 @@ void FNavigationToolMarker::SetInTime(const FFrameNumber& InTime)
 	}
 
 	const TArray<FMovieSceneMarkedFrame>& MarkedFrames = MovieScene->GetMarkedFrames();
+	if (!MarkedFrames.IsValidIndex(MarkedFrameIndex))
+	{
+		return;
+	}
+
 	FMovieSceneMarkedFrame& MarkedFrame = const_cast<FMovieSceneMarkedFrame&>(MarkedFrames[MarkedFrameIndex]);
 
 	if (InTime == MarkedFrame.FrameNumber)
