[SequenceNavigator] Hide toolbar button from unsupported sequence types

diff --git a/Engine/Plugins/Experimental/SequenceNavigator/Source/SequenceNavigator/Private/NavigationToolExtender.cpp b/Engine/Plugins/Experimental/SequenceNavigator/Source/SequenceNavigator/Private/NavigationToolExtender.cpp
--- a/Engine/Plugins/Experimental/SequenceNavigator/Source/SequenceNavigator/Private/NavigationToolExtender.cpp
+++ b/Engine/Plugins/Experimental/SequenceNavigator/Source/SequenceNavigator/Private/NavigationToolExtender.cpp
@@ -13,6 +13,7 @@
 #include "Providers/NavigationToolProvider.h"
 #include "SequenceNavigatorLog.h"
 #include "SequencerSettings.h"
+#include "SequencerToolMenuContext.h"
 
 namespace UE::SequenceNavigator
 {
@@ -421,13 +422,35 @@ void FNavigationToolExtender::AddSequencerToolBarExtension(const FNavigationTool
 
 	FToolMenuSection& NavigationToolSection = ToolMenu->FindOrAddSection(TEXT("NavigationTool"));
 
-	FToolMenuEntry ToggleTabVisibleEntry = FToolMenuEntry::InitToolBarButton(
-		FNavigationToolCommands::Get().ToggleToolTabVisible,
-		TAttribute<FText>(), TAttribute<FText>(),
-		FSlateIcon(FNavigationToolStyle::Get().GetStyleSetName(), TEXT("Icon.ToolBar")));
-	ToggleTabVisibleEntry.SetCommandList(InToolInstance.Instance->GetBaseCommandList());
-	ToggleTabVisibleEntry.StyleNameOverride = TEXT("SequencerToolbar");
-	NavigationToolSection.AddEntry(ToggleTabVisibleEntry);
+	NavigationToolSection.AddDynamicEntry(TEXT("NavigationTool")
+		, FNewToolMenuSectionDelegate::CreateLambda([ToolId](FToolMenuSection& InSection)
+		{
+			FNavigationToolInstance* const ToolInstance = Get().ToolInstances.Find(ToolId);
+			if (!ToolInstance)
+			{
+				return;
+			}
+
+			USequencerToolMenuContext* const ContextObject = InSection.FindContext<USequencerToolMenuContext>();
+			if (!ContextObject)
+			{
+				return;
+			}
+
+			const TSharedPtr<ISequencer> Sequencer = ContextObject->WeakSequencer.Pin();
+			if (Sequencer.IsValid() && !AnyProviderSupportsSequencer(ToolId, *Sequencer.Get()))
+			{
+				return;
+			}
+
+			FToolMenuEntry ToggleTabVisibleEntry = FToolMenuEntry::InitToolBarButton(
+				FNavigationToolCommands::Get().ToggleToolTabVisible,
+				TAttribute<FText>(), TAttribute<FText>(),
+				FSlateIcon(FNavigationToolStyle::Get().GetStyleSetName(), TEXT("Icon.ToolBar")));
+			ToggleTabVisibleEntry.SetCommandList(ToolInstance->Instance->GetBaseCommandList());
+			ToggleTabVisibleEntry.StyleNameOverride = TEXT("SequencerToolbar");
+			InSection.AddEntry(ToggleTabVisibleEntry);
+		}));
 }
 
 void FNavigationToolExtender::RemoveSequencerToolBarExtension(const FNavigationToolInstance& InToolInstance)
