[SequenceNavigator] Refactor to use Sequencer view models and type macros

diff --git a/Engine/Plugins/Experimental/SequenceNavigator/Source/SequenceNavigator/Private/Filters/Filters/NavigationToolFilter_Marks.cpp b/Engine/Plugins/Experimental/SequenceNavigator/Source/SequenceNavigator/Private/Filters/Filters/NavigationToolFilter_Marks.cpp
--- a/Engine/Plugins/Experimental/SequenceNavigator/Source/SequenceNavigator/Private/Filters/Filters/NavigationToolFilter_Marks.cpp
+++ b/Engine/Plugins/Experimental/SequenceNavigator/Source/SequenceNavigator/Private/Filters/Filters/NavigationToolFilter_Marks.cpp
@@ -10,7 +10,7 @@
 namespace UE::SequenceNavigator
 {
 
-using namespace UE::Sequencer;
+using namespace Sequencer;
 
 FNavigationToolFilter_Marks::FNavigationToolFilter_Marks(INavigationToolFilterBar& InFilterInterface, TSharedPtr<FFilterCategory> InCategory)
 	: FNavigationToolFilter(InFilterInterface, MoveTemp(InCategory))
@@ -42,9 +42,9 @@ FString FNavigationToolFilter_Marks::GetName() const
 	return StaticName();
 }
 
-bool FNavigationToolFilter_Marks::PassesFilter(FNavigationToolFilterType InItem) const
+bool FNavigationToolFilter_Marks::PassesFilter(const FNavigationToolViewModelPtr InItem) const
 {
-	const FNavigationToolSequence* const SequenceItem = InItem->CastTo<FNavigationToolSequence>();
+	const TViewModelPtr<FNavigationToolSequence> SequenceItem = InItem.ImplicitCast();
 	if (!SequenceItem)
 	{
 		return false;
diff --git a/Engine/Plugins/Experimental/SequenceNavigator/Source/SequenceNavigator/Private/Filters/Filters/NavigationToolFilter_Playhead.cpp b/Engine/Plugins/Experimental/SequenceNavigator/Source/SequenceNavigator/Private/Filters/Filters/NavigationToolFilter_Playhead.cpp
--- a/Engine/Plugins/Experimental/SequenceNavigator/Source/SequenceNavigator/Private/Filters/Filters/NavigationToolFilter_Playhead.cpp
+++ b/Engine/Plugins/Experimental/SequenceNavigator/Source/SequenceNavigator/Private/Filters/Filters/NavigationToolFilter_Playhead.cpp
@@ -13,7 +13,7 @@
 namespace UE::SequenceNavigator
 {
 
-using namespace UE::Sequencer;
+using namespace Sequencer;
 
 FNavigationToolFilter_Playhead::FNavigationToolFilter_Playhead(INavigationToolFilterBar& InFilterInterface, TSharedPtr<FFilterCategory> InCategory)
 	: FNavigationToolFilter(InFilterInterface, MoveTemp(InCategory))
@@ -89,15 +89,13 @@ FString FNavigationToolFilter_Playhead::GetName() const
 	return StaticName();
 }
 
-bool FNavigationToolFilter_Playhead::PassesFilter(FNavigationToolFilterType InItem) const
+bool FNavigationToolFilter_Playhead::PassesFilter(const FNavigationToolViewModelPtr InItem) const
 {
-	const IPlayheadExtension* const PlayheadItem = InItem->CastTo<IPlayheadExtension>();
-	if (!PlayheadItem)
+	if (const TViewModelPtr<IPlayheadExtension> PlayheadItem = InItem.ImplicitCast())
 	{
-		return false;
+		return PlayheadItem->ContainsPlayhead() != EItemContainsPlayhead::None;
 	}
-
-	return PlayheadItem->ContainsPlayhead() != EItemContainsPlayhead::None;
+	return false;
 }
 
 void FNavigationToolFilter_Playhead::OnPlayEvent()
diff --git a/Engine/Plugins/Experimental/SequenceNavigator/Source/SequenceNavigator/Private/Filters/Filters/NavigationToolFilter_Text.cpp b/Engine/Plugins/Experimental/SequenceNavigator/Source/SequenceNavigator/Private/Filters/Filters/NavigationToolFilter_Text.cpp
--- a/Engine/Plugins/Experimental/SequenceNavigator/Source/SequenceNavigator/Private/Filters/Filters/NavigationToolFilter_Text.cpp
+++ b/Engine/Plugins/Experimental/SequenceNavigator/Source/SequenceNavigator/Private/Filters/Filters/NavigationToolFilter_Text.cpp
@@ -15,7 +15,7 @@
 namespace UE::SequenceNavigator
 {
 
-using namespace UE::Sequencer;
+using namespace Sequencer;
 
 FNavigationToolFilter_Text::FNavigationToolFilter_Text(INavigationToolFilterBar& InFilterInterface)
 	: FNavigationToolFilter(InFilterInterface, nullptr)
@@ -64,7 +64,7 @@ FString FNavigationToolFilter_Text::GetName() const
 	return TEXT("TextFilter");
 }
 
-bool FNavigationToolFilter_Text::PassesFilter(const FNavigationToolItemPtr InItem) const
+bool FNavigationToolFilter_Text::PassesFilter(const FNavigationToolViewModelPtr InItem) const
 {
 	for (const TSharedRef<FNavigationToolFilterTextExpressionContext>& TextFilterExpressionContext : TextFilterExpressionContexts)
 	{
diff --git a/Engine/Plugins/Experimental/SequenceNavigator/Source/SequenceNavigator/Private/Filters/Filters/NavigationToolFilter_Unbound.cpp b/Engine/Plugins/Experimental/SequenceNavigator/Source/SequenceNavigator/Private/Filters/Filters/NavigationToolFilter_Unbound.cpp
--- a/Engine/Plugins/Experimental/SequenceNavigator/Source/SequenceNavigator/Private/Filters/Filters/NavigationToolFilter_Unbound.cpp
+++ b/Engine/Plugins/Experimental/SequenceNavigator/Source/SequenceNavigator/Private/Filters/Filters/NavigationToolFilter_Unbound.cpp
@@ -9,7 +9,7 @@
 namespace UE::SequenceNavigator
 {
 
-using namespace UE::Sequencer;
+using namespace Sequencer;
 
 FNavigationToolFilter_Unbound::FNavigationToolFilter_Unbound(INavigationToolFilterBar& InFilterInterface, TSharedPtr<FFilterCategory> InCategory)
 	: FNavigationToolFilter(InFilterInterface, MoveTemp(InCategory))
@@ -41,9 +41,9 @@ FString FNavigationToolFilter_Unbound::GetName() const
 	return StaticName();
 }
 
-bool FNavigationToolFilter_Unbound::PassesFilter(FNavigationToolFilterType InItem) const
+bool FNavigationToolFilter_Unbound::PassesFilter(const FNavigationToolViewModelPtr InItem) const
 {
-	FNavigationToolBinding* const BindingItem = InItem->CastTo<FNavigationToolBinding>();
+	const TViewModelPtr<FNavigationToolBinding> BindingItem = InItem.ImplicitCast();
 	if (!BindingItem)
 	{
 		return false;
diff --git a/Engine/Plugins/Experimental/SequenceNavigator/Source/SequenceNavigator/Private/Filters/Filters/NavigationToolFilters.cpp b/Engine/Plugins/Experimental/SequenceNavigator/Source/SequenceNavigator/Private/Filters/Filters/NavigationToolFilters.cpp
--- a/Engine/Plugins/Experimental/SequenceNavigator/Source/SequenceNavigator/Private/Filters/Filters/NavigationToolFilters.cpp
+++ b/Engine/Plugins/Experimental/SequenceNavigator/Source/SequenceNavigator/Private/Filters/Filters/NavigationToolFilters.cpp
@@ -14,7 +14,7 @@
 namespace UE::SequenceNavigator
 {
 
-using namespace UE::Sequencer;
+using namespace Sequencer;
 
 FNavigationToolFilter_Sequence::FNavigationToolFilter_Sequence(INavigationToolFilterBar& InFilterInterface, TSharedPtr<FFilterCategory> InCategory)
 	: FNavigationToolFilter_ItemType<FNavigationToolSequence>(InFilterInterface, InCategory)
diff --git a/Engine/Plugins/Experimental/SequenceNavigator/Source/SequenceNavigator/Private/Filters/NavigationToolFilterBar.cpp b/Engine/Plugins/Experimental/SequenceNavigator/Source/SequenceNavigator/Private/Filters/NavigationToolFilterBar.cpp
--- a/Engine/Plugins/Experimental/SequenceNavigator/Source/SequenceNavigator/Private/Filters/NavigationToolFilterBar.cpp
+++ b/Engine/Plugins/Experimental/SequenceNavigator/Source/SequenceNavigator/Private/Filters/NavigationToolFilterBar.cpp
@@ -17,6 +17,7 @@
 #include "Framework/Application/SlateApplication.h"
 #include "Framework/MultiBox/SToolBarButtonBlock.h"
 #include "ISequencer.h"
+#include "Items/NavigationToolTreeRoot.h"
 #include "Menus/NavigationToolFilterMenu.h"
 #include "NavigationTool.h"
 #include "NavigationToolSettings.h"
@@ -32,7 +33,7 @@
 namespace UE::SequenceNavigator
 {
 
-using namespace UE::Sequencer;
+using namespace Sequencer;
 
 const FName FNavigationToolFilterBar::SharedIdentifier = TEXT("SharedNavigationToolFilter");
 
@@ -81,7 +82,7 @@ void FNavigationToolFilterBar::Init()
 	CreateCustomTextFiltersFromConfig();
 }
 
-TSharedPtr<ICustomTextFilter<FNavigationToolItemPtr>> FNavigationToolFilterBar::CreateTextFilter()
+TSharedPtr<ICustomTextFilter<FNavigationToolViewModelPtr>> FNavigationToolFilterBar::CreateTextFilter()
 {
 	return MakeShared<FNavigationToolFilter_CustomText>(*this);
 }
@@ -1147,7 +1148,7 @@ TSharedRef<FFilterCategory> FNavigationToolFilterBar::GetMiscCategory() const
 	return MiscCategory;
 }
 
-bool FNavigationToolFilterBar::PassesAnyCommonFilter(const FNavigationToolItemPtr& InNode)
+bool FNavigationToolFilterBar::PassesAnyCommonFilter(const FNavigationToolViewModelPtr& InNode)
 {
 	bool bPassedAnyFilters = false;
 	bool bAnyFilterActive = false;
@@ -1178,7 +1179,7 @@ bool FNavigationToolFilterBar::PassesAnyCommonFilter(const FNavigationToolItemPt
 	return bPassedAnyFilters;
 }
 
-bool FNavigationToolFilterBar::PassesAllCustomTextFilters(const FNavigationToolItemPtr& InNode)
+bool FNavigationToolFilterBar::PassesAllCustomTextFilters(const FNavigationToolViewModelPtr& InNode)
 {
 	for (const TSharedRef<FNavigationToolFilter_CustomText>& Filter : CustomTextFilters)
 	{
@@ -1204,6 +1205,12 @@ const FNavigationToolFilterData& FNavigationToolFilterBar::FilterNodes()
 		return FilterData;
 	}
 
+	const FNavigationToolViewModelPtr RootItem = Tool.GetTreeRoot().Pin();
+    if (!RootItem.IsValid())
+    {
+    	return FilterData;
+    }
+
 	const UNavigationToolSettings* const ToolSettings = GetDefault<UNavigationToolSettings>();
 	if (!ToolSettings)
 	{
@@ -1215,7 +1222,7 @@ const FNavigationToolFilterData& FNavigationToolFilterBar::FilterNodes()
 
 	// Loop through all nodes and filter recursively
 	const bool bHasActiveFilter = HasAnyFilterActive();
-	for (const FNavigationToolItemPtr& RootNode : Tool.GetTreeRoot()->GetChildren())
+	for (const FNavigationToolViewModelWeakPtr& RootNode : RootItem->GetChildren())
 	{
 		FilterNodesRecursive(*ToolView, *ToolSettings, bHasActiveFilter, RootNode);
 	}
@@ -1241,8 +1248,14 @@ TArray<TSharedRef<ISequencerTextFilterExpressionContext>> FNavigationToolFilterB
 bool FNavigationToolFilterBar::FilterNodesRecursive(INavigationToolView& InToolView
 	, const UNavigationToolSettings& InSettings
 	, const bool bInHasActiveFilter
-	, const FNavigationToolItemPtr& InStartNode)
+	, const FNavigationToolViewModelWeakPtr& InWeakStartNode)
 {
+	const FNavigationToolViewModelPtr StartNode = InWeakStartNode.Pin();
+	if (!StartNode.IsValid())
+	{
+		return false;
+	}
+
 	/**
 	 * Main Filtering Logic
 	 *
@@ -1254,9 +1267,9 @@ bool FNavigationToolFilterBar::FilterNodesRecursive(INavigationToolView& InToolV
 	bool bAnyChildPassed = false;
 
 	// Child nodes should always be processed, as they may force their parents to pass
-	for (const FNavigationToolItemPtr& Node : InStartNode->GetChildren())
+	for (const FNavigationToolViewModelWeakPtr& WeakNode : StartNode->GetChildren())
 	{
-		if (FilterNodesRecursive(InToolView, InSettings, bInHasActiveFilter, Node))
+		if (FilterNodesRecursive(InToolView, InSettings, bInHasActiveFilter, WeakNode))
 		{
 			bAnyChildPassed = true;
 		}
@@ -1268,13 +1281,13 @@ bool FNavigationToolFilterBar::FilterNodesRecursive(INavigationToolView& InToolV
 	// Early out if no filter
 	if (!bInHasActiveFilter)
 	{
-		FilterData.FilterInNode(InStartNode);
+		FilterData.FilterInNode(StartNode);
 		return false;
 	}
 
-	const bool bPassedTextFilter = !TextFilter->IsActive() || TextFilter->PassesFilter(InStartNode);
-	const bool bPassedAnyCommonFilters = PassesAnyCommonFilter(InStartNode);
-	const bool bPassedAnyCustomTextFilters = PassesAllCustomTextFilters(InStartNode);
+	const bool bPassedTextFilter = !TextFilter->IsActive() || TextFilter->PassesFilter(StartNode);
+	const bool bPassedAnyCommonFilters = PassesAnyCommonFilter(StartNode);
+	const bool bPassedAnyCustomTextFilters = PassesAllCustomTextFilters(StartNode);
 
 	const bool bAllFiltersPassed = bPassedTextFilter
 		&& bPassedAnyCommonFilters
@@ -1284,15 +1297,15 @@ bool FNavigationToolFilterBar::FilterNodesRecursive(INavigationToolView& InToolV
 	{
 		if (InSettings.ShouldAutoExpandNodesOnFilterPass())
 		{
-			InToolView.SetParentItemExpansions(InStartNode, true);
+			InToolView.SetParentItemExpansions(InWeakStartNode, true);
 		}
 
-		FilterData.FilterInNodeWithAncestors(InStartNode);
+		FilterData.FilterInNodeWithAncestors(InWeakStartNode);
 		return true;
 	}
 
 	// After child nodes are processed, fail anything that didn't pass
-	FilterData.FilterOutNode(InStartNode);
+	FilterData.FilterOutNode(InWeakStartNode);
 	return false;
 }
 
diff --git a/Engine/Plugins/Experimental/SequenceNavigator/Source/SequenceNavigator/Private/Filters/NavigationToolFilterData.cpp b/Engine/Plugins/Experimental/SequenceNavigator/Source/SequenceNavigator/Private/Filters/NavigationToolFilterData.cpp
--- a/Engine/Plugins/Experimental/SequenceNavigator/Source/SequenceNavigator/Private/Filters/NavigationToolFilterData.cpp
+++ b/Engine/Plugins/Experimental/SequenceNavigator/Source/SequenceNavigator/Private/Filters/NavigationToolFilterData.cpp
@@ -1,12 +1,12 @@
 ﻿// Copyright Epic Games, Inc. All Rights Reserved.
 
 #include "Filters/NavigationToolFilterData.h"
-#include "Items/INavigationToolItem.h"
+#include "Items/NavigationToolItem.h"
 
 namespace UE::SequenceNavigator
 {
 
-using namespace UE::Sequencer;
+using namespace Sequencer;
 
 FNavigationToolFilterData::FNavigationToolFilterData(const FString& InRawFilterText)
 	: RawFilterText(InRawFilterText)
@@ -60,31 +60,32 @@ void FNavigationToolFilterData::IncrementTotalNodeCount()
 	++TotalNodeCount;
 }
 
-void FNavigationToolFilterData::FilterInNode(const FNavigationToolItemPtr& InNode)
+void FNavigationToolFilterData::FilterInNode(const FNavigationToolViewModelWeakPtr& InWeakNode)
 {
-	FilterInNodes.Add(InNode);
+	FilterInNodes.Add(InWeakNode.Pin());
 }
 
-void FNavigationToolFilterData::FilterOutNode(const FNavigationToolItemPtr& InNode)
+void FNavigationToolFilterData::FilterOutNode(const FNavigationToolViewModelWeakPtr& InWeakNode)
 {
-	const FSetElementId ElementId = FilterInNodes.FindId(InNode);
+	const FSetElementId ElementId = FilterInNodes.FindId(InWeakNode.Pin());
 	if (ElementId.IsValidId())
 	{
 		FilterInNodes.Remove(ElementId);
 	}
 }
 
-void FNavigationToolFilterData::FilterInParentChildNodes(const FNavigationToolItemPtr& InNode
+void FNavigationToolFilterData::FilterInParentChildNodes(const FNavigationToolViewModelWeakPtr& InWeakNode
 	, const bool bInIncludeSelf, const bool bInIncludeParents, const bool bInIncludeChildren)
 {
-	if (!InNode.IsValid())
+	const FNavigationToolViewModelPtr Item = InWeakNode.Pin();
+	if (!Item.IsValid())
 	{
 		return;
 	}
 
 	if (bInIncludeParents)
 	{
-		FNavigationToolItemPtr CurrentParent = InNode->GetParent();
+		FNavigationToolViewModelPtr CurrentParent = Item->GetParent();
 		while (CurrentParent.IsValid())
 		{
 			FilterInNode(CurrentParent);
@@ -99,29 +100,29 @@ void FNavigationToolFilterData::FilterInParentChildNodes(const FNavigationToolIt
 
 	if (bInIncludeSelf)
 	{
-		FilterInNode(InNode);
+		FilterInNode(InWeakNode);
 	}
 
 	if (bInIncludeChildren)
 	{
-		auto AddAllChildren = [this](const FNavigationToolItemPtr& InNode)
+		auto AddAllChildren = [this](const FNavigationToolViewModelPtr& InNode)
 		{
-			for (const FNavigationToolItemPtr& Child : InNode->GetChildren())
+			for (const FNavigationToolViewModelWeakPtr& WeakChild : InNode->GetChildren())
 			{
-				FilterInNode(Child);
+				FilterInNode(WeakChild);
 			}
 		};
 
-		for (const FNavigationToolItemPtr& Child : InNode->GetChildren())
+		for (const FNavigationToolViewModelWeakPtr& WeakChild : Item->GetChildren())
 		{
-			AddAllChildren(Child);
+			AddAllChildren(WeakChild.Pin());
 		}
 	}
 }
 
-void FNavigationToolFilterData::FilterInNodeWithAncestors(const FNavigationToolItemPtr& InNode)
+void FNavigationToolFilterData::FilterInNodeWithAncestors(const FNavigationToolViewModelWeakPtr& InWeakNode)
 {
-	FilterInParentChildNodes(InNode, true, true, false);
+	FilterInParentChildNodes(InWeakNode.Pin(), true, true, false);
 }
 
 bool FNavigationToolFilterData::ContainsFilterInNodes(const FNavigationToolFilterData& InOtherData) const
@@ -129,14 +130,14 @@ bool FNavigationToolFilterData::ContainsFilterInNodes(const FNavigationToolFilte
 	return FilterInNodes.Includes(InOtherData.FilterInNodes);
 }
 
-bool FNavigationToolFilterData::IsFilteredIn(const FNavigationToolItemPtr& InNode) const
+bool FNavigationToolFilterData::IsFilteredIn(const FNavigationToolViewModelWeakPtr& InWeakNode) const
 {
-	return FilterInNodes.Contains(InNode);
+	return FilterInNodes.Contains(InWeakNode.Pin());
 }
 
-bool FNavigationToolFilterData::IsFilteredOut(const FNavigationToolItemPtr& InNode) const
+bool FNavigationToolFilterData::IsFilteredOut(const FNavigationToolViewModelWeakPtr& InWeakNode) const
 {
-	return !FilterInNodes.Contains(InNode);
+	return !FilterInNodes.Contains(InWeakNode.Pin());
 }
 
 } // namespace UE::SequenceNavigator
diff --git a/Engine/Plugins/Experimental/SequenceNavigator/Source/SequenceNavigator/Private/Filters/TextExpressions/NavigationToolFilterTextExpressionContext.cpp b/Engine/Plugins/Experimental/SequenceNavigator/Source/SequenceNavigator/Private/Filters/TextExpressions/NavigationToolFilterTextExpressionContext.cpp
--- a/Engine/Plugins/Experimental/SequenceNavigator/Source/SequenceNavigator/Private/Filters/TextExpressions/NavigationToolFilterTextExpressionContext.cpp
+++ b/Engine/Plugins/Experimental/SequenceNavigator/Source/SequenceNavigator/Private/Filters/TextExpressions/NavigationToolFilterTextExpressionContext.cpp
@@ -2,19 +2,19 @@
 
 #include "Filters/TextExpressions/NavigationToolFilterTextExpressionContext.h"
 #include "Filters/INavigationToolFilterBar.h"
-#include "Items/INavigationToolItem.h"
+#include "Items/NavigationToolItem.h"
 
 namespace UE::SequenceNavigator
 {
 
-using namespace UE::Sequencer;
+using namespace Sequencer;
 
 FNavigationToolFilterTextExpressionContext::FNavigationToolFilterTextExpressionContext(INavigationToolFilterBar& InFilterInterface)
 	: FilterInterface(InFilterInterface)
 {
 }
 
-void FNavigationToolFilterTextExpressionContext::SetFilterItem(const FNavigationToolItemPtr& InFilterItem)
+void FNavigationToolFilterTextExpressionContext::SetFilterItem(const FNavigationToolViewModelPtr& InFilterItem)
 {
 	WeakFilterItem = InFilterItem;
 }
@@ -22,7 +22,7 @@ void FNavigationToolFilterTextExpressionContext::SetFilterItem(const FNavigation
 bool FNavigationToolFilterTextExpressionContext::TestBasicStringExpression(const FTextFilterString& InValue
 	, const ETextFilterTextComparisonMode InTextComparisonMode) const
 {
-	const TSharedPtr<INavigationToolItem> FilterItem = WeakFilterItem.Pin();
+	const FNavigationToolViewModelPtr FilterItem = WeakFilterItem.Pin();
 	if (!FilterItem.IsValid())
 	{
 		return false;
@@ -37,7 +37,7 @@ bool FNavigationToolFilterTextExpressionContext::TestBasicStringExpression(const
 		return true;
 	}
 
-	for (const FNavigationToolItemPtr& ParentItem : FilterItem->GetParents())
+	for (const FNavigationToolViewModelPtr& ParentItem : FilterItem->GetParents())
 	{
 		Label = ParentItem->GetDisplayName().ToString();
 		if (TextFilterUtils::TestBasicStringExpression(Label, InValue, TextComparisonMode))
@@ -54,11 +54,6 @@ bool FNavigationToolFilterTextExpressionContext::TestComplexExpression(const FNa
 	, const ETextFilterComparisonOperation InComparisonOperation
 	, const ETextFilterTextComparisonMode InTextComparisonMode) const
 {
-	if (!WeakFilterItem.IsValid())
-	{
-		return false;
-	}
-
 	const TSet<FName> Keys = GetKeys();
 	if (!Keys.IsEmpty() && !Keys.Contains(InKey))
 	{
diff --git a/Engine/Plugins/Experimental/SequenceNavigator/Source/SequenceNavigator/Private/Filters/TextExpressions/NavigationToolFilterTextExpression_Name.cpp b/Engine/Plugins/Experimental/SequenceNavigator/Source/SequenceNavigator/Private/Filters/TextExpressions/NavigationToolFilterTextExpression_Name.cpp
--- a/Engine/Plugins/Experimental/SequenceNavigator/Source/SequenceNavigator/Private/Filters/TextExpressions/NavigationToolFilterTextExpression_Name.cpp
+++ b/Engine/Plugins/Experimental/SequenceNavigator/Source/SequenceNavigator/Private/Filters/TextExpressions/NavigationToolFilterTextExpression_Name.cpp
@@ -1,14 +1,14 @@
 ﻿// Copyright Epic Games, Inc. All Rights Reserved.
 
 #include "NavigationToolFilterTextExpression_Name.h"
-#include "Items/INavigationToolItem.h"
+#include "Items/NavigationToolItem.h"
 
 #define LOCTEXT_NAMESPACE "NavigationToolFilterTextExpression_Name"
 
 namespace UE::SequenceNavigator
 {
 
-using namespace UE::Sequencer;
+using namespace Sequencer;
 
 FNavigationToolFilterTextExpression_Name::FNavigationToolFilterTextExpression_Name(INavigationToolFilterBar& InFilterInterface)
 	: FNavigationToolFilterTextExpressionContext(InFilterInterface)
@@ -40,7 +40,7 @@ bool FNavigationToolFilterTextExpression_Name::TestComplexExpression(const FName
 		return true;
 	}
 
-	if (const TSharedPtr<INavigationToolItem> FilterItem = WeakFilterItem.Pin())
+	if (const FNavigationToolViewModelPtr FilterItem = WeakFilterItem.Pin())
 	{
 		const FString Label = FilterItem->GetDisplayName().ToString();
 		const bool bPassed = TextFilterUtils::TestComplexExpression(Label, InValue, InComparisonOperation, InTextComparisonMode);
diff --git a/Engine/Plugins/Experimental/SequenceNavigator/Source/SequenceNavigator/Private/Filters/TextExpressions/NavigationToolFilterTextExpression_Unbound.cpp b/Engine/Plugins/Experimental/SequenceNavigator/Source/SequenceNavigator/Private/Filters/TextExpressions/NavigationToolFilterTextExpression_Unbound.cpp
--- a/Engine/Plugins/Experimental/SequenceNavigator/Source/SequenceNavigator/Private/Filters/TextExpressions/NavigationToolFilterTextExpression_Unbound.cpp
+++ b/Engine/Plugins/Experimental/SequenceNavigator/Source/SequenceNavigator/Private/Filters/TextExpressions/NavigationToolFilterTextExpression_Unbound.cpp
@@ -7,7 +7,7 @@
 namespace UE::SequenceNavigator
 {
 
-using namespace UE::Sequencer;
+using namespace Sequencer;
 
 FNavigationToolFilterTextExpression_Unbound::FNavigationToolFilterTextExpression_Unbound(INavigationToolFilterBar& InFilterInterface)
 	: FNavigationToolFilterTextExpressionContext(InFilterInterface)
@@ -39,7 +39,7 @@ bool FNavigationToolFilterTextExpression_Unbound::TestComplexExpression(const FN
 		return true;
 	}
 
-	/*if (const TSharedPtr<INavigationToolItem> FilterItem = WeakFilterItem.Pin())
+	/*if (const FNavigationToolViewModelPtr FilterItem = WeakFilterItem.Pin())
 	{
 		const FString Label = FilterItem->GetDisplayName().ToString();
 		const bool bPassed = TextFilterUtils::TestComplexExpression(Label, InValue, InComparisonOperation, InTextComparisonMode);
diff --git a/Engine/Plugins/Experimental/SequenceNavigator/Source/SequenceNavigator/Private/ItemActions/NavigationToolAddItem.cpp b/Engine/Plugins/Experimental/SequenceNavigator/Source/SequenceNavigator/Private/ItemActions/NavigationToolAddItem.cpp
--- a/Engine/Plugins/Experimental/SequenceNavigator/Source/SequenceNavigator/Private/ItemActions/NavigationToolAddItem.cpp
+++ b/Engine/Plugins/Experimental/SequenceNavigator/Source/SequenceNavigator/Private/ItemActions/NavigationToolAddItem.cpp
@@ -2,6 +2,7 @@
 
 #include "ItemActions/NavigationToolAddItem.h"
 #include "Items/INavigationToolItem.h"
+#include "Items/NavigationToolTreeRoot.h"
 #include "NavigationTool.h"
 #include "NavigationToolScopedSelection.h"
 #include "Widgets/Views/STableRow.h"
@@ -9,6 +10,10 @@
 namespace UE::SequenceNavigator
 {
 
+using namespace Sequencer;
+
+UE_SEQUENCER_DEFINE_CASTABLE(FNavigationToolAddItem)
+
 FNavigationToolAddItem::FNavigationToolAddItem(const FNavigationToolAddItemParams& InAddItemParams)
 	: AddParams(InAddItemParams)
 {
@@ -21,7 +26,14 @@ bool FNavigationToolAddItem::ShouldTransact() const
 
 void FNavigationToolAddItem::Execute(FNavigationTool& InTool)
 {
-	if (!AddParams.Item.IsValid())
+	const FNavigationToolViewModelPtr TreeRoot = InTool.GetTreeRoot().Pin();
+	if (!TreeRoot.IsValid())
+	{
+		return;
+	}
+
+	const FNavigationToolViewModelPtr ItemToAdd = AddParams.WeakItem.Pin();
+	if (!ItemToAdd.IsValid())
 	{
 		return;
 	}
@@ -29,68 +41,67 @@ void FNavigationToolAddItem::Execute(FNavigationTool& InTool)
 	// Try to Create Children on Find
 	if (EnumHasAnyFlags(AddParams.Flags, ENavigationToolAddItemFlags::AddChildren))
 	{
-		constexpr bool bRecursiveFind = true;
-		TArray<FNavigationToolItemPtr> Children;
-		AddParams.Item->FindValidChildren(Children, bRecursiveFind);
+		TArray<FNavigationToolViewModelWeakPtr> WeakChildren;
+		ItemToAdd->FindValidChildren(WeakChildren, /*bInRecursive=*/true);
 	}
 
-	const FNavigationToolItemPtr ParentItem = AddParams.Item->GetParent();
+	const FNavigationToolViewModelPtr ParentItem = ItemToAdd->GetParent();
 
 	// If this array has elements in it, then we need to stop a circular dependency from forming
-	const TArray<FNavigationToolItemPtr> PathToRelativeItem = AddParams.Item->FindPath({ AddParams.RelativeItem });
+	const TArray<FNavigationToolViewModelPtr> PathToRelativeItem = ItemToAdd->FindPath({ AddParams.WeakRelativeItem });
 	if (PathToRelativeItem.Num() > 0 && ParentItem.IsValid())
 	{
 		FNavigationToolAddItemParams CircularSolverParams;
-		CircularSolverParams.Item             = PathToRelativeItem[0];
-		CircularSolverParams.RelativeItem     = AddParams.Item;
+		CircularSolverParams.WeakItem         = PathToRelativeItem[0];
+		CircularSolverParams.WeakRelativeItem = AddParams.WeakItem;
 		CircularSolverParams.RelativeDropZone = EItemDropZone::AboveItem;
 		CircularSolverParams.Flags            = AddParams.Flags;
 
 		ParentItem->AddChild(MoveTemp(CircularSolverParams));
 	}
 
-	if (AddParams.RelativeItem.IsValid())
+	if (const FNavigationToolViewModelPtr RelativeItem = AddParams.WeakRelativeItem.Pin())
 	{
-		const FNavigationToolItemPtr RelativeItemParent = AddParams.RelativeItem->GetParent();
+		const FNavigationToolViewModelPtr RelativeItemParent = RelativeItem->GetParent();
 
-		// If it's onto item, the Relative Item is going to be the Parent
+		// If it's onto item, the relative item is going to be the parent
 		if (!AddParams.RelativeDropZone.IsSet() || AddParams.RelativeDropZone == EItemDropZone::OntoItem)
 		{
-			//If the Relative Item is Onto and it's the same as the Current Parent, shift Item up in the Hierarchy
-			//(as long as the parent is a valid one)
-			if (AddParams.RelativeItem == ParentItem && RelativeItemParent)
+			// If the relative item is onto and it's the same as the current parent, shift Item up in the hierarchy
+			// (as long as the parent is valid)
+			if (RelativeItem == ParentItem && RelativeItemParent)
 			{
 				AddParams.RelativeDropZone = EItemDropZone::BelowItem;
 				RelativeItemParent->AddChild(AddParams);
 			}
 			else
 			{
-				AddParams.RelativeItem->AddChild(AddParams);
+				RelativeItem->AddChild(AddParams);
 			}
 		}
-		//else we place it as a Sibling to the Relative Item
+		// Else we place it as a Sibling to the Relative Item
 		else if (RelativeItemParent)
 		{
 			RelativeItemParent->AddChild(AddParams);
 		}
-		//if no parent, then add it to the tree root
+		// If no parent, then add it to the tree root
 		else
 		{
-			InTool.GetTreeRoot()->AddChild(AddParams);
+			TreeRoot->AddChild(AddParams);
 		}
 	}
 	else
 	{
-		//If no Relative Item, add to tree root
-		InTool.GetTreeRoot()->AddChild(AddParams);
+		// If no relative item, add to tree root
+		TreeRoot->AddChild(AddParams);
 	}
 
 	const FNavigationToolScopedSelection ScopedSelection(*InTool.GetSequencer(), ENavigationToolScopedSelectionPurpose::Read);
 
-	// Automatically select item if it's selected
-	if (AddParams.Item->IsSelected(ScopedSelection))
+	// Automatically select the item if it's selected
+	if (ItemToAdd->IsSelected(ScopedSelection))
 	{
-		// Select in Navigation Tool but don't signal selection as we already have it selected in Mode Tools
+		// Select in Navigation Tool but don't signal selection as we already have it selected in mode tools
 		AddParams.Flags = ENavigationToolAddItemFlags::Select;
 		AddParams.SelectionFlags &= ~ENavigationToolItemSelectionFlags::SignalSelectionChange;
 	}
@@ -102,7 +113,7 @@ void FNavigationToolAddItem::Execute(FNavigationTool& InTool)
 
 	if (EnumHasAnyFlags(AddParams.Flags, ENavigationToolAddItemFlags::Select))
 	{
-		InTool.SelectItems({ AddParams.Item }, AddParams.SelectionFlags);
+		InTool.SelectItems({ ItemToAdd }, AddParams.SelectionFlags);
 	}
 
 	InTool.SetToolModified();
@@ -110,13 +121,16 @@ void FNavigationToolAddItem::Execute(FNavigationTool& InTool)
 
 void FNavigationToolAddItem::OnObjectsReplaced(const TMap<UObject*, UObject*>& InReplacementMap, bool bRecursive)
 {
-	if (AddParams.Item.IsValid())
+	using namespace Sequencer;
+
+	if (const FNavigationToolViewModelPtr ItemToAdd = AddParams.WeakItem.Pin())
 	{
-		AddParams.Item->OnObjectsReplaced(InReplacementMap, bRecursive);
+		ItemToAdd->OnObjectsReplaced(InReplacementMap, bRecursive);
 	}
-	if (AddParams.RelativeItem.IsValid())
+
+	if (const FNavigationToolViewModelPtr RelativeItem = AddParams.WeakRelativeItem.Pin())
 	{
-		AddParams.RelativeItem->OnObjectsReplaced(InReplacementMap, bRecursive);
+		RelativeItem->OnObjectsReplaced(InReplacementMap, bRecursive);
 	}
 }
 
diff --git a/Engine/Plugins/Experimental/SequenceNavigator/Source/SequenceNavigator/Private/ItemActions/NavigationToolRemoveItem.cpp b/Engine/Plugins/Experimental/SequenceNavigator/Source/SequenceNavigator/Private/ItemActions/NavigationToolRemoveItem.cpp
--- a/Engine/Plugins/Experimental/SequenceNavigator/Source/SequenceNavigator/Private/ItemActions/NavigationToolRemoveItem.cpp
+++ b/Engine/Plugins/Experimental/SequenceNavigator/Source/SequenceNavigator/Private/ItemActions/NavigationToolRemoveItem.cpp
@@ -9,86 +9,99 @@
 namespace UE::SequenceNavigator
 {
 
+using namespace Sequencer;
+
+UE_SEQUENCER_DEFINE_CASTABLE(FNavigationToolRemoveItem)
+
 FNavigationToolRemoveItem::FNavigationToolRemoveItem(const FNavigationToolRemoveItemParams& InRemoveItemParams)
 	: RemoveParams(InRemoveItemParams)
 {
-	if (RemoveParams.Item.IsValid())
+	if (const FNavigationToolViewModelPtr Item = RemoveParams.WeakItem.Pin())
 	{
-		RemoveParams.Item->AddFlags(ENavigationToolItemFlags::PendingRemoval);
+		Item->AddFlags(ENavigationToolItemFlags::PendingRemoval);
 	}
 }
 
 void FNavigationToolRemoveItem::Execute(FNavigationTool& InTool)
 {
-	if (RemoveParams.Item.IsValid())
+	using namespace Sequencer;
+
+	const FNavigationToolViewModelPtr Item = RemoveParams.WeakItem.Pin();
+	if (!Item.IsValid())
 	{
-		FNavigationToolItemFlagGuard Guard(RemoveParams.Item, ENavigationToolItemFlags::IgnorePendingKill);
+		return;
+	}
 
-		//Copy the Array since we may modify it below on Add/Remove Child
-		TArray<FNavigationToolItemPtr> Children = RemoveParams.Item->GetChildren();
+	const FNavigationToolItemFlagGuard Guard(Item, ENavigationToolItemFlags::IgnorePendingKill);
 
-		FNavigationToolItemPtr Parent       = RemoveParams.Item->GetParent();
-		FNavigationToolItemPtr RelativeItem = RemoveParams.Item;
+	// Copy the array since we may modify it below on add/remove child
+	TArray<FNavigationToolViewModelWeakPtr> WeakChildren = Item->GetChildren();
 
-		//Search the lowest parent that is not pending removal
-		while (Parent.IsValid() && Parent->HasAnyFlags(ENavigationToolItemFlags::PendingRemoval))
-		{
-			RelativeItem = Parent;
-			Parent       = Parent->GetParent();
-		}
+	FNavigationToolViewModelPtr Parent = Item->GetParent();
+	FNavigationToolViewModelPtr RelativeItem = Item;
 
-		//Reparent the Item's Children to the Valid Parent found above
-		if (Parent.IsValid())
-		{
-			FNavigationToolAddItemParams ReparentParams;
-			ReparentParams.RelativeItem             = RelativeItem;
-			ReparentParams.RelativeDropZone         = EItemDropZone::BelowItem;
-			ReparentParams.Flags                    = ENavigationToolAddItemFlags::Select;
+	// Search the lowest parent that is not pending removal
+	while (Parent.IsValid() && Parent->HasAnyFlags(ENavigationToolItemFlags::PendingRemoval))
+	{
+		RelativeItem = Parent;
+		Parent = Parent->GetParent();
+	}
 
-			TArray<FNavigationToolItemPtr> ItemsToReparent;
+	// Reparent the item's children to the valid parent found above
+	if (Parent.IsValid())
+	{
+		FNavigationToolAddItemParams ReparentParams;
+		ReparentParams.WeakRelativeItem = RelativeItem;
+		ReparentParams.RelativeDropZone = EItemDropZone::BelowItem;
+		ReparentParams.Flags = ENavigationToolAddItemFlags::Select;
 
-			ItemsToReparent.Append(Children);
-			while (ItemsToReparent.Num() > 0)
-			{
-				ReparentParams.Item = ItemsToReparent.Pop();
+		TArray<FNavigationToolViewModelWeakPtr> ItemsToReparent;
 
-				//Then Parent it to the Item's Parent.
-				//If we couldn't add the Child, then it means either the Child Item itself is invalid
-				//or the underlying info (e.g. object) is invalid (e.g. actor pending kill)
-				if (Parent->AddChild(ReparentParams) == false)
+		ItemsToReparent.Append(WeakChildren);
+		while (ItemsToReparent.Num() > 0)
+		{
+			ItemsToReparent.Pop();
+			ReparentParams.WeakItem = ItemsToReparent.Pop();
+
+			// Then parent it to the item's parent.
+			// If we couldn't add the Child, then it means either the child item itself is invalid
+			// or the underlying info (e.g. object) is invalid (e.g. actor pending kill).
+			if (!Parent->AddChild(ReparentParams))
+			{
+				// So try reparenting the children of this invalid child (since this invalid child will be removed)
+				if (const FNavigationToolViewModelPtr ChildItem = ReparentParams.WeakItem.Pin())
 				{
-					//so try reparenting the children of this invalid child (since this invalid child will be removed)
-					if (ReparentParams.Item.IsValid())
-					{
-						ItemsToReparent.Append(ReparentParams.Item->GetChildren());
-					}
+					ItemsToReparent.Append(ChildItem->GetChildren());
 				}
 			}
+		}
 
-			//In case the Parent is still the same, Remove Child Item, else the Parent is going to be removed anyways
-			if (Parent == RemoveParams.Item->GetParent())
-			{
-				Parent->RemoveChild(RemoveParams.Item);
-			}
+		// In case the parent is still the same, remove child item, else the parent is going to be removed anyways
+		if (Parent == Item->GetParent())
+		{
+			Parent->RemoveChild(Item);
 		}
-		else
+	}
+	else
+	{
+		for (const FNavigationToolViewModelWeakPtr& WeakChild : WeakChildren)
 		{
-			for (const FNavigationToolItemPtr& Child : Children)
+			if (const FNavigationToolViewModelPtr ChildItem = WeakChild.Pin())
 			{
-				RemoveParams.Item->RemoveChild(Child);
+				Item->RemoveChild(ChildItem);
 			}
 		}
-
-		InTool.UnregisterItem(RemoveParams.Item->GetItemId());
-		InTool.SetToolModified();
 	}
+
+	InTool.UnregisterItem(Item->GetItemId());
+	InTool.SetToolModified();
 }
 
-void FNavigationToolRemoveItem::OnObjectsReplaced(const TMap<UObject*, UObject*>& InReplacementMap, bool bRecursive)
+void FNavigationToolRemoveItem::OnObjectsReplaced(const TMap<UObject*, UObject*>& InReplacementMap, const bool bInRecursive)
 {
-	if (RemoveParams.Item.IsValid())
+	if (const FNavigationToolViewModelPtr Item = RemoveParams.WeakItem.Pin())
 	{
-		RemoveParams.Item->OnObjectsReplaced(InReplacementMap, bRecursive);
+		Item->OnObjectsReplaced(InReplacementMap, bInRecursive);
 	}
 }
 
diff --git a/Engine/Plugins/Experimental/SequenceNavigator/Source/SequenceNavigator/Private/ItemProxies/NavigationToolComponentProxy.cpp b/Engine/Plugins/Experimental/SequenceNavigator/Source/SequenceNavigator/Private/ItemProxies/NavigationToolComponentProxy.cpp
--- a/Engine/Plugins/Experimental/SequenceNavigator/Source/SequenceNavigator/Private/ItemProxies/NavigationToolComponentProxy.cpp
+++ b/Engine/Plugins/Experimental/SequenceNavigator/Source/SequenceNavigator/Private/ItemProxies/NavigationToolComponentProxy.cpp
@@ -9,8 +9,10 @@
 namespace UE::SequenceNavigator
 {
 
-FNavigationToolComponentProxy::FNavigationToolComponentProxy(INavigationTool& InTool, const FNavigationToolItemPtr& InParentItem)
-	: Super(InTool, InParentItem)
+UE_SEQUENCER_DEFINE_CASTABLE(FNavigationToolComponentProxy)
+
+FNavigationToolComponentProxy::FNavigationToolComponentProxy(INavigationTool& InTool, const FNavigationToolViewModelPtr& InParentItem)
+	: FNavigationToolItemProxy(InTool, InParentItem)
 {
 }
 
@@ -36,7 +38,8 @@ ENavigationToolItemViewMode FNavigationToolComponentProxy::GetSupportedViewModes
 	return ~ENavigationToolItemViewMode::ItemTree | ENavigationToolItemViewMode::HorizontalItemList;
 }
 
-void FNavigationToolComponentProxy::GetProxiedItems(const TSharedRef<INavigationToolItem>& InParent, TArray<FNavigationToolItemPtr>& OutChildren, bool bInRecursive)
+void FNavigationToolComponentProxy::GetProxiedItems(const FNavigationToolViewModelPtr& InParent
+	, TArray<FNavigationToolViewModelWeakPtr>& OutWeakChildren, const bool bInRecursive)
 {
 }
 
diff --git a/Engine/Plugins/Experimental/SequenceNavigator/Source/SequenceNavigator/Private/Items/NavigationToolActor.cpp b/Engine/Plugins/Experimental/SequenceNavigator/Source/SequenceNavigator/Private/Items/NavigationToolActor.cpp
--- a/Engine/Plugins/Experimental/SequenceNavigator/Source/SequenceNavigator/Private/Items/NavigationToolActor.cpp
+++ b/Engine/Plugins/Experimental/SequenceNavigator/Source/SequenceNavigator/Private/Items/NavigationToolActor.cpp
@@ -19,17 +19,21 @@
 namespace UE::SequenceNavigator
 {
 
+using namespace Sequencer;
+
+UE_SEQUENCER_DEFINE_CASTABLE(FNavigationToolActor)
+
 FNavigationToolActor::FNavigationToolActor(INavigationTool& InTool
-	, const FNavigationToolItemPtr& InParentItem
+	, const FNavigationToolViewModelPtr& InParentItem
 	, const TSharedPtr<FNavigationToolSequence>& InParentSequenceItem
 	, const FMovieSceneBinding& InBinding)
-	: Super(InTool, InParentItem, InParentSequenceItem, InBinding)
+	: FNavigationToolBinding(InTool, InParentItem, InParentSequenceItem, InBinding)
 {
 }
 
-void FNavigationToolActor::FindChildren(TArray<FNavigationToolItemPtr>& OutChildren, const bool bInRecursive)
+void FNavigationToolActor::FindChildren(TArray<FNavigationToolViewModelWeakPtr>& OutWeakChildren, const bool bInRecursive)
 {
-	Super::FindChildren(OutChildren, bInRecursive);
+	FNavigationToolBinding::FindChildren(OutWeakChildren, bInRecursive);
 
 	const AActor* const ThisActor = GetActor();
 	if (!ThisActor)
@@ -37,7 +41,7 @@ void FNavigationToolActor::FindChildren(TArray<FNavigationToolItemPtr>& OutChild
 		return;
 	}
 
-	const TSharedPtr<FNavigationToolSequence> ParentSequenceItem = WeakParentSequenceItem.Pin();
+	const TViewModelPtr<FNavigationToolSequence> ParentSequenceItem = WeakParentSequenceItem.Pin();
 	if (!ParentSequenceItem.IsValid())
 	{
 		return;
@@ -61,9 +65,16 @@ void FNavigationToolActor::FindChildren(TArray<FNavigationToolItemPtr>& OutChild
 		return;
 	}
 
-	const TSharedRef<FNavigationToolActor> SharedThisRef = SharedThis(this);
+	const TSharedPtr<FNavigationToolProvider> Provider = GetProvider();
+	if (!Provider.IsValid())
+	{
+		return;
+	}
 
+	const TSharedRef<FNavigationToolProvider> ProviderRef = Provider.ToSharedRef();
+	const FNavigationToolViewModelPtr ParentItem = AsItemViewModel();
 	const TArray<FMovieSceneBinding> Bindings = ParentSequenceItem->GetSortedBindings();
+
 	for (const FMovieSceneBinding& CurrentBinding : Bindings)
 	{
 		const FGuid& CurrentBindingGuid = CurrentBinding.GetObjectGuid();
@@ -72,16 +83,16 @@ void FNavigationToolActor::FindChildren(TArray<FNavigationToolItemPtr>& OutChild
 		{
 			const TArrayView<TWeakObjectPtr<>> WeakBoundObjects = ResolveBoundObjects(*Sequencer, Sequence, CurrentBindingGuid);
 			const TWeakObjectPtr<> CurrentBoundObject = WeakBoundObjects.IsEmpty() ? nullptr : WeakBoundObjects[0];
-
 			const UActorComponent* const CurrentBoundActorComponent = Cast<UActorComponent>(CurrentBoundObject);
+
 			if (CurrentBoundActorComponent && CurrentBoundActorComponent->GetOwner() == ThisActor)
 			{
-				const FNavigationToolItemPtr NewItem = Tool.FindOrAdd<FNavigationToolComponent>(GetProvider()
-					, SharedThisRef, ParentSequenceItem, CurrentBinding);
-				OutChildren.Add(NewItem);
+				const FNavigationToolViewModelPtr NewItem = Tool.FindOrAdd<FNavigationToolComponent>(ProviderRef
+					, ParentItem, ParentSequenceItem, CurrentBinding);
+				OutWeakChildren.Add(NewItem);
 				if (bInRecursive)
 				{
-					NewItem->FindChildren(OutChildren, bInRecursive);
+					NewItem->FindChildren(OutWeakChildren, bInRecursive);
 				}
 			}
 		}
@@ -115,7 +126,7 @@ TArray<FName> FNavigationToolActor::GetTags() const
 	{
 		return UnderlyingActor->Tags;
 	}
-	return Super::GetTags();
+	return FNavigationToolBinding::GetTags();
 }
 
 bool FNavigationToolActor::GetVisibility() const
@@ -137,28 +148,28 @@ void FNavigationToolActor::OnVisibilityChanged(const bool bInNewVisibility)
 
 bool FNavigationToolActor::CanRename() const
 {
-	return Super::CanRename() && GetActor() != nullptr;
+	return FNavigationToolBinding::CanRename() && !GetActor();
 }
 
-bool FNavigationToolActor::Rename(const FString& InName)
+void FNavigationToolActor::Rename(const FText& InNewName)
 {
 	AActor* const UnderlyingActor = GetActor();
-	const bool bIsActorLabelEditable = UnderlyingActor && UnderlyingActor->IsActorLabelEditable();
-
-	bool bRenamed = false;
-	
-	if (bIsActorLabelEditable && !InName.Equals(UnderlyingActor->GetActorLabel(), ESearchCase::CaseSensitive))
+	if (!UnderlyingActor || !UnderlyingActor->IsActorLabelEditable())
 	{
-		const FScopedTransaction Transaction(LOCTEXT("NavigationToolRenameActor", "Rename Actor"));
+		return;
+	}
 
-		FActorLabelUtilities::RenameExistingActor(UnderlyingActor, InName);
+	const FString NewNameString = InNewName.ToString();
+	if (NewNameString.Equals(UnderlyingActor->GetActorLabel(), ESearchCase::CaseSensitive))
+	{
+		return;
+	}
 
-		Super::Rename(InName);
+	const FScopedTransaction Transaction(LOCTEXT("NavigationToolRenameActor", "Rename Actor"));
 
-		bRenamed = true;
-	}
+	FActorLabelUtilities::RenameExistingActor(UnderlyingActor, NewNameString);
 
-	return bRenamed;
+	FNavigationToolBinding::Rename(InNewName);
 }
 
 AActor* FNavigationToolActor::GetActor() const
diff --git a/Engine/Plugins/Experimental/SequenceNavigator/Source/SequenceNavigator/Private/Items/NavigationToolBinding.cpp b/Engine/Plugins/Experimental/SequenceNavigator/Source/SequenceNavigator/Private/Items/NavigationToolBinding.cpp
--- a/Engine/Plugins/Experimental/SequenceNavigator/Source/SequenceNavigator/Private/Items/NavigationToolBinding.cpp
+++ b/Engine/Plugins/Experimental/SequenceNavigator/Source/SequenceNavigator/Private/Items/NavigationToolBinding.cpp
@@ -22,13 +22,15 @@
 namespace UE::SequenceNavigator
 {
 
-using namespace UE::Sequencer;
+using namespace Sequencer;
+
+UE_SEQUENCER_DEFINE_CASTABLE(FNavigationToolBinding)
 
 FNavigationToolBinding::FNavigationToolBinding(INavigationTool& InTool
-	, const FNavigationToolItemPtr& InParentItem
+	, const FNavigationToolViewModelPtr& InParentItem
 	, const TSharedPtr<FNavigationToolSequence>& InParentSequenceItem
 	, const FMovieSceneBinding& InBinding)
-	: Super(InTool, InParentItem)
+	: FNavigationToolItem(InTool, InParentItem)
 	, WeakParentSequenceItem(InParentSequenceItem)
 	, Binding(InBinding)
 {
@@ -52,11 +54,11 @@ bool FNavigationToolBinding::IsAllowedInTool() const
 	return IsItemValid();
 }
 
-void FNavigationToolBinding::FindChildren(TArray<FNavigationToolItemPtr>& OutChildren, const bool bInRecursive)
+void FNavigationToolBinding::FindChildren(TArray<FNavigationToolViewModelWeakPtr>& OutWeakChildren, const bool bInRecursive)
 {
 	CacheBoundObject();
 
-	Super::FindChildren(OutChildren, bInRecursive);
+	FNavigationToolItem::FindChildren(OutWeakChildren, bInRecursive);
 }
 
 FText FNavigationToolBinding::GetDisplayName() const
@@ -114,7 +116,7 @@ void FNavigationToolBinding::Select(FNavigationToolScopedSelection& InSelection)
 	
 void FNavigationToolBinding::OnSelect()
 {
-	Super::OnSelect();
+	FNavigationToolItem::OnSelect();
 }
 
 void FNavigationToolBinding::OnDoubleClick()
@@ -156,11 +158,56 @@ bool FNavigationToolBinding::CanRename() const
 	//return WeakBoundObject.IsValid();
 }
 
-bool FNavigationToolBinding::Rename(const FString& InName)
+void FNavigationToolBinding::Rename(const FText& InNewName)
 {
-	Binding.SetName(InName);
+	UMovieSceneSequence* const Sequence = GetSequence();
+	if (!Sequence)
+	{
+		return;
+	}
 
-	return false;
+	UMovieScene* const MovieScene = Sequence->GetMovieScene();
+	if (!MovieScene)
+	{
+		return;
+	}
+
+	const TSharedPtr<ISequencer> Sequencer = Tool.GetSequencer();
+	if (!Sequencer.IsValid())
+	{
+		return;
+	}
+
+	const FGuid ObjectBindingID = Binding.GetObjectGuid();
+	const FString NewNameString = InNewName.ToString();
+
+	const FScopedTransaction Transaction(LOCTEXT("SetTrackName", "Set Track Name"));
+
+	// Modify the movie scene so that it gets marked dirty and renames are saved consistently.
+	MovieScene->Modify();
+
+	// If there is only one binding, set the name of the bound actor
+	const TArrayView<TWeakObjectPtr<>> Objects = Sequencer->FindObjectsInCurrentSequence(ObjectBindingID);
+	if (Objects.Num() == 1)
+	{
+		if (AActor* const Actor = Cast<AActor>(Objects[0].Get()))
+		{
+			Actor->SetActorLabel(NewNameString);
+		}
+	}
+
+	if (FMovieSceneSpawnable* const Spawnable = MovieScene->FindSpawnable(ObjectBindingID))
+	{
+		Spawnable->SetName(NewNameString);
+	}
+	else if (FMovieScenePossessable* const Possessable = MovieScene->FindPossessable(ObjectBindingID))
+	{
+		Possessable->SetName(NewNameString);
+	}
+	else
+	{
+		MovieScene->SetObjectDisplayName(ObjectBindingID, InNewName);
+	}
 }
 /*
 EItemSequenceLockState FNavigationToolBinding::GetLockState() const
@@ -240,6 +287,8 @@ UObject* FNavigationToolBinding::GetCachedBoundObject() const
 
 UObject* FNavigationToolBinding::CacheBoundObject()
 {
+	using namespace ItemUtils;
+
 	const TSharedPtr<ISequencer> Sequencer = Tool.GetSequencer();
 	if (!Sequencer.IsValid())
 	{
diff --git a/Engine/Plugins/Experimental/SequenceNavigator/Source/SequenceNavigator/Private/Items/NavigationToolComponent.cpp b/Engine/Plugins/Experimental/SequenceNavigator/Source/SequenceNavigator/Private/Items/NavigationToolComponent.cpp
--- a/Engine/Plugins/Experimental/SequenceNavigator/Source/SequenceNavigator/Private/Items/NavigationToolComponent.cpp
+++ b/Engine/Plugins/Experimental/SequenceNavigator/Source/SequenceNavigator/Private/Items/NavigationToolComponent.cpp
@@ -5,6 +5,7 @@
 #include "Components/SceneComponent.h"
 #include "Items/NavigationToolItemProxy.h"
 #include "NavigationTool.h"
+#include "NavigationToolSettings.h"
 #include "ScopedTransaction.h"
 #include "Styling/StyleColors.h"
 #include "Widgets/Columns/SNavigationToolLabelComponent.h"
@@ -14,26 +15,28 @@
 namespace UE::SequenceNavigator
 {
 
+UE_SEQUENCER_DEFINE_CASTABLE(FNavigationToolComponent)
+
 FNavigationToolComponent::FNavigationToolComponent(INavigationTool& InTool
-	, const FNavigationToolItemPtr& InParentItem
+	, const FNavigationToolViewModelPtr& InParentItem
 	, const TSharedPtr<FNavigationToolSequence>& InParentSequenceItem
 	, const FMovieSceneBinding& InBinding)
-	: Super(InTool, InParentItem, InParentSequenceItem, InBinding)
+	: FNavigationToolBinding(InTool, InParentItem, InParentSequenceItem, InBinding)
 {
 }
 
-void FNavigationToolComponent::FindChildren(TArray<FNavigationToolItemPtr>& OutChildren, const bool bInRecursive)
+void FNavigationToolComponent::FindChildren(TArray<FNavigationToolViewModelWeakPtr>& OutWeakChildren, const bool bInRecursive)
 {
-	Super::FindChildren(OutChildren, bInRecursive);
+	FNavigationToolBinding::FindChildren(OutWeakChildren, bInRecursive);
 }
 
 void FNavigationToolComponent::GetItemProxies(TArray<TSharedPtr<FNavigationToolItemProxy>>& OutItemProxies)
 {
-	Super::GetItemProxies(OutItemProxies);
+	FNavigationToolBinding::GetItemProxies(OutItemProxies);
 
 	if (UPrimitiveComponent* const PrimitiveComponent = Cast<UPrimitiveComponent>(GetComponent()))
 	{
-		if (const TSharedPtr<FNavigationToolItemProxy> MaterialItemProxy = Tool.GetOrCreateItemProxy<FNavigationToolItemProxy>(SharedThis(this)))
+		if (const TSharedPtr<FNavigationToolItemProxy> MaterialItemProxy = Tool.GetOrCreateItemProxy<FNavigationToolItemProxy>(AsItemViewModel()))
 		{
 			OutItemProxies.Add(MaterialItemProxy);
 		}
@@ -64,7 +67,7 @@ ENavigationToolItemViewMode FNavigationToolComponent::GetSupportedViewModes(cons
 
 TSharedRef<SWidget> FNavigationToolComponent::GenerateLabelWidget(const TSharedRef<SNavigationToolTreeRow>& InRow)
 {
-	return SNew(SNavigationToolLabelComponent, SharedThis(this), InRow);
+	return SNew(SNavigationToolLabelComponent, AsItemViewModel(), InRow);
 }
 
 FLinearColor FNavigationToolComponent::GetItemTintColor() const
@@ -78,7 +81,7 @@ TArray<FName> FNavigationToolComponent::GetTags() const
 	{
 		return UnderlyingComponent->ComponentTags;
 	}
-	return Super::GetTags();
+	return FNavigationToolBinding::GetTags();
 }
 
 bool FNavigationToolComponent::GetVisibility() const
@@ -100,32 +103,30 @@ void FNavigationToolComponent::OnVisibilityChanged(const bool bInNewVisibility)
 
 bool FNavigationToolComponent::CanRename() const
 {
-	return Super::CanRename() && GetComponent() != nullptr;
+	return FNavigationToolBinding::CanRename() && !GetComponent();
 }
 
-bool FNavigationToolComponent::Rename(const FString& InName)
+void FNavigationToolComponent::Rename(const FText& InNewName)
 {
 	UActorComponent* const UnderlyingComponent = GetComponent();
 	if (!UnderlyingComponent)
 	{
-		return false;
+		return;
 	}
 
-	bool bRenamed = false;
+	const FString NewNameString = InNewName.ToString();
 
-	if (!InName.Equals(UnderlyingComponent->GetName(), ESearchCase::CaseSensitive))
+	if (NewNameString.Equals(UnderlyingComponent->GetName(), ESearchCase::CaseSensitive))
 	{
-		const FScopedTransaction Transaction(LOCTEXT("NavigationToolRenameComponent", "Rename Component"));
-
-		UnderlyingComponent->Modify();
-		UnderlyingComponent->Rename(*InName);
+		return;
+	}
 
-		Super::Rename(InName);
+	const FScopedTransaction Transaction(LOCTEXT("RenameComponent", "Rename Component"));
 
-		bRenamed = true;
-	}
+	UnderlyingComponent->Modify();
+	UnderlyingComponent->Rename(*NewNameString);
 
-	return bRenamed;
+	FNavigationToolBinding::Rename(InNewName);
 }
 
 UActorComponent* FNavigationToolComponent::GetComponent() const
diff --git a/Engine/Plugins/Experimental/SequenceNavigator/Source/SequenceNavigator/Private/Items/NavigationToolItem.cpp b/Engine/Plugins/Experimental/SequenceNavigator/Source/SequenceNavigator/Private/Items/NavigationToolItem.cpp
--- a/Engine/Plugins/Experimental/SequenceNavigator/Source/SequenceNavigator/Private/Items/NavigationToolItem.cpp
+++ b/Engine/Plugins/Experimental/SequenceNavigator/Source/SequenceNavigator/Private/Items/NavigationToolItem.cpp
@@ -9,6 +9,7 @@
 #include "Items/NavigationToolItemUtils.h"
 #include "Misc/Optional.h"
 #include "NavigationTool.h"
+#include "NavigationToolSettings.h"
 #include "NavigationToolView.h"
 #include "Providers/NavigationToolProvider.h"
 #include "Styling/StyleColors.h"
@@ -21,7 +22,11 @@
 namespace UE::SequenceNavigator
 {
 
-FNavigationToolItem::FNavigationToolItem(INavigationTool& InTool, const FNavigationToolItemPtr& InParentItem)
+using namespace Sequencer;
+
+UE_SEQUENCER_DEFINE_CASTABLE(FNavigationToolItem)
+
+FNavigationToolItem::FNavigationToolItem(INavigationTool& InTool, const FNavigationToolViewModelPtr& InParentItem)
 	: Tool(InTool)
 	, WeakParent(InParentItem)
 {
@@ -44,45 +49,54 @@ bool FNavigationToolItem::IsItemValid() const
 
 void FNavigationToolItem::RefreshChildren()
 {
-	TArray<FNavigationToolItemPtr> FoundChildren;
-	FindValidChildren(FoundChildren, /*bRecursiveFind=*/false);
+	using namespace ItemUtils;
+
+	TArray<FNavigationToolViewModelWeakPtr> WeakFoundChildren;
+	FindValidChildren(WeakFoundChildren, /*bRecursiveFind=*/false);
 
-	TArray<FNavigationToolItemPtr> Sortable;
-	TArray<FNavigationToolItemPtr> Unsortable;
-	SplitSortableAndUnsortableItems(FoundChildren, Sortable, Unsortable);
+	TArray<FNavigationToolViewModelWeakPtr> Sortable;
+	TArray<FNavigationToolViewModelWeakPtr> Unsortable;
+	SplitSortableAndUnsortableItems(WeakFoundChildren, Sortable, Unsortable);
 
-	// Start with all Sortable/Unsortable Items, and remove every item seen by iterating Children
-	TSet<FNavigationToolItemPtr> NewSortableChildren(Sortable);
-	TSet<FNavigationToolItemPtr> NewUnsortableChildren(Unsortable);
+	// Start with all sortable/unsortable items, and remove every item seen by iterating children
+	TSet<FNavigationToolViewModelWeakPtr> NewSortableChildren(Sortable);
+	TSet<FNavigationToolViewModelWeakPtr> NewUnsortableChildren(Unsortable);
 	
-	// Remove items from "Children" that were not present in the Sortable Found Children (we'll add non-sortable later)
-	// Result is have Children only contain Items that existed previously
-	for (TArray<FNavigationToolItemPtr>::TIterator ItemIter(Children); ItemIter; ++ItemIter)
+	// Remove items from "WeakChildren" that were not present in the sortable found children (we'll add non-sortable later)
+	// Result is have WeakChildren only contain Items that existed previously
+	TSet<FNavigationToolViewModelWeakPtr> ItemsToRemove;
+
+	for (const FNavigationToolViewModelWeakPtr& WeakChildItem : WeakChildren)
 	{
-		const FNavigationToolItemPtr Item(*ItemIter);
+		const FNavigationToolViewModelPtr Item = WeakChildItem.Pin();
 
 		if (!Item.IsValid() || NewUnsortableChildren.Contains(Item))
 		{
-			ItemIter.RemoveCurrent();
+			ItemsToRemove.Add(WeakChildItem);
 		}
 		else if (!NewSortableChildren.Contains(Item) || !Item->IsItemValid())
 		{
 			Item->SetParent(nullptr);
-			ItemIter.RemoveCurrent();
+			ItemsToRemove.Add(WeakChildItem);
 		}
 
 		NewSortableChildren.Remove(Item);
 		NewUnsortableChildren.Remove(Item);
 	}
 
-	// Find Children for New Children in case these new children have grand children
-	// Note: This does not affect any of the current containers. It's just called for discovery
-	auto FindGrandChildren = [](const TSet<FNavigationToolItemPtr>& InChildren)
+	WeakChildren.RemoveAll([&ItemsToRemove](const FNavigationToolViewModelWeakPtr& InWeakItem)
 		{
-			for (const FNavigationToolItemPtr& Child : InChildren)
+			return ItemsToRemove.Contains(InWeakItem);
+		});
+
+	// Find children for new children in case these new children have grand children.
+	// Note: This does not affect any of the current containers. It's just called for discovery.
+	auto FindGrandChildren = [](const TSet<FNavigationToolViewModelWeakPtr>& InWeakChildren)
+		{
+			for (const FNavigationToolViewModelWeakPtr& WeakChild : InWeakChildren)
 			{
-				TArray<FNavigationToolItemPtr> GrandChildren;
-				Child->FindValidChildren(GrandChildren, /*bRecursiveFind=*/true);
+				TArray<FNavigationToolViewModelWeakPtr> WeakGrandChildren;
+				WeakChild.Pin()->FindValidChildren(WeakGrandChildren, /*bRecursiveFind=*/true);
 			}
 		};
 
@@ -92,101 +106,169 @@ void FNavigationToolItem::RefreshChildren()
 	// After removing Children not present in Sortable
 	// Children should either be equal in size with Sortable (which means no new sortable children were added)
 	// or Sortable has more entries which means there are new items to add
-	if (Sortable.Num() > Children.Num())
+	if (Sortable.Num() > WeakChildren.Num())
 	{
 		check(!NewSortableChildren.IsEmpty());
 		HandleNewSortableChildren(NewSortableChildren.Array());
 	}
 
-	// Rearrange so that Children are arranged like so:
+	// Rearrange so that children are arranged like so:
 	// [Unsortable Children][Sortable Children]
-	Unsortable.Append(MoveTemp(Children));
-	Children = MoveTemp(Unsortable);
+	Unsortable.Append(MoveTemp(WeakChildren));
+	WeakChildren = MoveTemp(Unsortable);
 
-	// Update the Parents of every Child in the List
-	const FNavigationToolItemPtr This = SharedThis(this);
-	for (const FNavigationToolItemPtr& Child : Children)
+	// Update the parents of every child in the list
+	const FNavigationToolViewModelPtr ThisRef = SharedThis(this);
+	for (const FNavigationToolViewModelWeakPtr& WeakChild : WeakChildren)
 	{
-		Child->SetParent(This);
+		WeakChild.Pin()->SetParent(ThisRef);
 	}
 }
 
 void FNavigationToolItem::ResetChildren()
 {
-	for (const FNavigationToolItemPtr& Item : GetChildren())
+	using namespace Sequencer;
+
+	for (const FNavigationToolViewModelWeakPtr& WeakItem : GetChildren())
 	{
-		if (Item.IsValid())
+		if (const FNavigationToolViewModelPtr Item = WeakItem.Pin())
 		{
 			Item->SetParent(nullptr);
 		}
 	}
+
 	GetChildrenMutable().Reset();	
 };
 
-void FNavigationToolItem::FindChildren(TArray<FNavigationToolItemPtr>& OutChildren, const bool bInRecursive)
+void FNavigationToolItem::FindChildren(TArray<FNavigationToolViewModelWeakPtr>& OutWeakChildren
+	, const bool bInRecursive)
 {
 	FNavigationTool& ToolPrivate = static_cast<FNavigationTool&>(Tool);
 
+	const FNavigationToolViewModelPtr ThisRef = SharedThis(this);
+
 	TArray<TSharedPtr<FNavigationToolItemProxy>> ItemProxies;
-	ToolPrivate.GetItemProxiesForItem(SharedThis(this), ItemProxies);
-	OutChildren.Reserve(OutChildren.Num() + ItemProxies.Num());
+	ToolPrivate.GetItemProxiesForItem(ThisRef, ItemProxies);
+	OutWeakChildren.Reserve(OutWeakChildren.Num() + ItemProxies.Num());
 
 	for (const TSharedPtr<FNavigationToolItemProxy>& ItemProxy : ItemProxies)
 	{
-		OutChildren.Add(ItemProxy);
+		OutWeakChildren.Add(ItemProxy);
 		if (bInRecursive)
 		{
-			ItemProxy->FindChildren(OutChildren, bInRecursive);
+			ItemProxy->FindChildren(OutWeakChildren, bInRecursive);
 		}
 	}
 }
 
-void FNavigationToolItem::FindValidChildren(TArray<FNavigationToolItemPtr>& OutChildren, bool bInRecursive)
+void FNavigationToolItem::FindValidChildren(TArray<FNavigationToolViewModelWeakPtr>& OutWeakChildren
+	, const bool bInRecursive)
 {
-	FindChildren(OutChildren, bInRecursive);
+	FindChildren(OutWeakChildren, bInRecursive);
 
-	OutChildren.RemoveAll(
-		[](const FNavigationToolItemPtr& InItem)
+	OutWeakChildren.RemoveAll([](const FNavigationToolViewModelWeakPtr& InWeakItem)
 		{
-			return !InItem.IsValid() || !InItem->IsAllowedInTool();
+			const FNavigationToolViewModelPtr Item = InWeakItem.Pin();
+			return !Item.IsValid() || !Item->IsAllowedInTool();
 		});
 }
 
-TArray<FNavigationToolItemPtr> FNavigationToolItem::FindPath(const TArray<FNavigationToolItemPtr>& InItems) const
+TArray<FNavigationToolViewModelPtr> FNavigationToolItem::FindPath(const TArray<FNavigationToolViewModelWeakPtr>& InWeakItems) const
 {
-	const TSharedPtr<const INavigationToolItem> This = SharedThis(this);
-	TArray<FNavigationToolItemPtr> Path;
-	for (const FNavigationToolItemPtr& Item : InItems)
+	using namespace Sequencer;
+
+	//const TSharedPtr<const INavigationToolItem> This = SharedThis(this);
+	TArray<FNavigationToolViewModelPtr> Path;
+
+	for (const FNavigationToolViewModelWeakPtr& WeakItem : InWeakItems)
 	{
 		Path.Reset();
-		FNavigationToolItemPtr CurrentItem = Item;
+		FNavigationToolViewModelPtr CurrentItem = WeakItem.Pin();
+
 		while (CurrentItem.IsValid())
 		{
-			if (This == CurrentItem)
+			//if (This == *CurrentItem)
+			if (this == CurrentItem.Get())
 			{
 				Algo::Reverse(Path);
 				return Path;
 			}
+
 			Path.Add(CurrentItem);
 			CurrentItem = CurrentItem->GetParent();
 		}
 	}
-	return TArray<FNavigationToolItemPtr>();
+
+	return TArray<FNavigationToolViewModelPtr>();
+}
+
+FString FNavigationToolItem::GetFullPath() const
+{
+	FString OutPath;
+
+	FNavigationToolViewModelPtr CurrentParent = GetParent();
+	while (CurrentParent.IsValid())
+	{
+		if (OutPath.IsEmpty())
+		{
+			OutPath = CurrentParent->GetItemId().GetStringId();
+		}
+		else
+		{
+			OutPath = CurrentParent->GetItemId().GetStringId() + TEXT(",") + OutPath;
+		}
+
+		CurrentParent = CurrentParent->GetParent();
+	}
+
+	return OutPath;
 }
 
-INavigationToolItem::IndexType FNavigationToolItem::GetChildIndex(const FNavigationToolItemPtr& ChildItem) const
+FNavigationToolItem::IndexType FNavigationToolItem::GetChildIndex(const FNavigationToolViewModelWeakPtr& InWeakChildItem) const
 {
-	return GetChildren().Find(ChildItem);
+	return GetChildren().Find(InWeakChildItem);
+}
+
+FNavigationToolViewModelWeakPtr FNavigationToolItem::GetChildAt(const IndexType InIndex) const
+{
+	const TArray<FNavigationToolViewModelWeakPtr>& ChildItems = GetChildren();
+	if (ChildItems.IsValidIndex(InIndex))
+	{
+		return ChildItems[InIndex].Pin();
+	}
+	return nullptr;
+}
+
+FNavigationToolViewModelPtr FNavigationToolItem::GetParent() const
+{
+	return FindAncestorOfType<FNavigationToolItem>();
+}
+
+TSet<FNavigationToolViewModelPtr> FNavigationToolItem::GetParents(const bool bInIncludeRoot) const
+{
+	TSet<FNavigationToolViewModelPtr> OutParents;
+
+	FNavigationToolViewModelPtr CurrentParent = GetParent();
+	while (CurrentParent.IsValid()
+		&& (bInIncludeRoot || CurrentParent->GetItemId() != FNavigationToolItemId::RootId))
+	{
+		OutParents.Add(CurrentParent);
+
+		CurrentParent = CurrentParent->GetParent();
+	}
+
+	return OutParents;
 }
 
-bool FNavigationToolItem::CanAddChild(const FNavigationToolItemPtr& InChild) const
+bool FNavigationToolItem::CanAddChild(const FNavigationToolViewModelPtr& InChild) const
 {
 	return InChild.IsValid();
 }
 
 bool FNavigationToolItem::AddChild(const FNavigationToolAddItemParams& InAddItemParams)
 {
-	if (CanAddChild(InAddItemParams.Item))
+	const FNavigationToolViewModelPtr Item = InAddItemParams.WeakItem.Pin();
+	if (CanAddChild(Item))
 	{
 		AddChildChecked(InAddItemParams);
 		return true;
@@ -196,14 +278,10 @@ bool FNavigationToolItem::AddChild(const FNavigationToolAddItemParams& InAddItem
 
 bool FNavigationToolItem::RemoveChild(const FNavigationToolRemoveItemParams& InRemoveItemParams)
 {
-	if (InRemoveItemParams.Item.IsValid())
-	{
-		return RemoveChildChecked(InRemoveItemParams);
-	}
-	return false;
+	return RemoveChildChecked(InRemoveItemParams);
 }
 
-void FNavigationToolItem::SetParent(FNavigationToolItemPtr InParent)
+void FNavigationToolItem::SetParent(FNavigationToolViewModelPtr InParent)
 {
 	//check that one of the parent's children is this
 	WeakParent = InParent;
@@ -213,6 +291,11 @@ ENavigationToolItemViewMode FNavigationToolItem::GetSupportedViewModes(const INa
 {
 	return InToolView.GetItemDefaultViewMode();
 }
+	
+bool FNavigationToolItem::IsViewModeSupported(const ENavigationToolItemViewMode InViewMode, const INavigationToolView& InToolView) const
+{
+	return EnumHasAnyFlags(InViewMode, GetSupportedViewModes(InToolView));
+}
 
 FNavigationToolItemId FNavigationToolItem::GetItemId() const
 {
@@ -231,7 +314,7 @@ FSlateColor FNavigationToolItem::GetItemLabelColor() const
 
 const FSlateBrush* FNavigationToolItem::GetIconBrush() const
 {
-	const FSlateIcon Icon = FNavigationToolExtender::FindOverrideIcon(SharedThis(this));
+	const FSlateIcon Icon = FNavigationToolExtender::FindOverrideIcon(AsItemViewModelConst());
 	if (Icon.IsSet())
 	{
 		return Icon.GetIcon();
@@ -281,7 +364,7 @@ bool FNavigationToolItem::HasAllFlags(ENavigationToolItemFlags Flags) const
 	return EnumHasAllFlags(ItemFlags, Flags);
 }
 
-TOptional<EItemDropZone> FNavigationToolItem::CanAcceptDrop(const FDragDropEvent& InDragDropEvent, EItemDropZone InDropZone)
+TOptional<EItemDropZone> FNavigationToolItem::CanAcceptDrop(const FDragDropEvent& InDragDropEvent, const EItemDropZone InDropZone)
 {
 	if (const TSharedPtr<FNavigationToolItemDragDropOp> ItemDragDropOp = InDragDropEvent.GetOperationAs<FNavigationToolItemDragDropOp>())
 	{
@@ -298,14 +381,12 @@ TOptional<EItemDropZone> FNavigationToolItem::CanAcceptDrop(const FDragDropEvent
 		ItemDragDropOp->CurrentIconBrush = FAppStyle::GetBrush(TEXT("Graph.ConnectorFeedback.Error"));
 	}
 
-	const TSharedRef<FNavigationToolItem> SharedThisRef = SharedThis(this);
-
 	TOptional<EItemDropZone> OutDropZone;
 
-	GetOwnerTool().ForEachProvider([this, &InDragDropEvent, InDropZone, &SharedThisRef, &OutDropZone]
+	GetOwnerTool().ForEachProvider([this, &InDragDropEvent, InDropZone, &OutDropZone]
 		(const TSharedRef<FNavigationToolProvider>& InProvider)
 		{
-			OutDropZone = InProvider->OnToolItemCanAcceptDrop(InDragDropEvent, InDropZone, SharedThisRef);
+			OutDropZone = InProvider->OnToolItemCanAcceptDrop(InDragDropEvent, InDropZone, SharedThis(this));
 			if (OutDropZone.IsSet())
 			{
 				return false;
@@ -316,7 +397,7 @@ TOptional<EItemDropZone> FNavigationToolItem::CanAcceptDrop(const FDragDropEvent
 	return OutDropZone;
 }
 
-FReply FNavigationToolItem::AcceptDrop(const FDragDropEvent& InDragDropEvent, EItemDropZone InDropZone)
+FReply FNavigationToolItem::AcceptDrop(const FDragDropEvent& InDragDropEvent, const EItemDropZone InDropZone)
 {
 	if (const TSharedPtr<FNavigationToolItemDragDropOp> ItemDragDropOp = InDragDropEvent.GetOperationAs<FNavigationToolItemDragDropOp>())
 	{
@@ -327,14 +408,12 @@ FReply FNavigationToolItem::AcceptDrop(const FDragDropEvent& InDragDropEvent, EI
 		}
 	}
 
-	const TSharedRef<FNavigationToolItem> SharedThisRef = SharedThis(this);
-
 	FReply OutReply = FReply::Unhandled();
 
-	GetOwnerTool().ForEachProvider([this, &InDragDropEvent, InDropZone, &SharedThisRef, &OutReply]
+	GetOwnerTool().ForEachProvider([this, &InDragDropEvent, InDropZone, &OutReply]
 		(const TSharedRef<FNavigationToolProvider>& InProvider)
 		{
-			OutReply = InProvider->OnToolItemAcceptDrop(InDragDropEvent, InDropZone, SharedThisRef);
+			OutReply = InProvider->OnToolItemAcceptDrop(InDragDropEvent, InDropZone, SharedThis(this));
 			if (OutReply.IsEventHandled())
 			{
 				return false;
@@ -345,6 +424,27 @@ FReply FNavigationToolItem::AcceptDrop(const FDragDropEvent& InDragDropEvent, EI
 	return OutReply;
 }
 
+bool FNavigationToolItem::IsIgnoringPendingKill() const
+{
+	return HasAllFlags(ENavigationToolItemFlags::IgnorePendingKill);
+}
+
+void FNavigationToolItem::OnObjectsReplaced(const TMap<UObject*, UObject*>& InReplacementMap, const bool bInRecursive)
+{
+	using namespace Sequencer;
+
+	if (bInRecursive)
+	{
+		for (const FNavigationToolViewModelWeakPtr& WeakChildItem : GetChildren())
+		{
+			if (const FNavigationToolViewModelPtr ChildItem = WeakChildItem.Pin())
+			{
+				ChildItem->OnObjectsReplaced(InReplacementMap, bInRecursive);
+			}
+		}
+	}
+}
+
 FLinearColor FNavigationToolItem::GetItemTintColor() const
 {
 	return FStyleColors::White.GetSpecifiedColor();
@@ -377,6 +477,23 @@ void FNavigationToolItem::SetColor(const TOptional<FColor>& InColor)
 	Tool.SetItemColor(SharedThis(this), InColor.Get(FColor()));
 }
 
+FNavigationToolSerializedItem FNavigationToolItem::MakeSerializedItem()
+{
+	return FNavigationToolSerializedItem(GetItemId().GetStringId());
+}
+
+FOutlinerSizing FNavigationToolItem::GetOutlinerSizing() const
+{
+	int32 Height = 0;
+	FNavigationToolViewModelPtr TopParent = GetParent();
+	while (TopParent.IsValid())
+	{
+		TopParent = TopParent->GetParent();
+		++Height;
+	}
+	return Height;
+}
+
 void FNavigationToolItem::SetProvider(const TWeakPtr<FNavigationToolProvider>& InWeakProvider)
 {
 	WeakProvider = InWeakProvider;
@@ -393,97 +510,116 @@ void FNavigationToolItem::RecalculateItemId()
 
 void FNavigationToolItem::AddChildChecked(const FNavigationToolAddItemParams& InAddItemParams)
 {
-	if (const FNavigationToolItemPtr OldParent = InAddItemParams.Item->GetParent())
+	using namespace Sequencer;
+
+	const FNavigationToolViewModelPtr ItemToAdd = InAddItemParams.WeakItem.Pin();
+	if (!ItemToAdd.IsValid())
 	{
-		// if we are adding the child and the old parent is this,
+		return;
+	}
+
+	if (const FNavigationToolViewModelPtr OldParent = ItemToAdd->GetParent())
+	{
+		// If we are adding the child and the old parent is this,
 		// then it means we're just rearranging, only remove from array
 		if (OldParent.Get() == this)
 		{
-			Children.Remove(InAddItemParams.Item);
+			WeakChildren.Remove(InAddItemParams.WeakItem);
 		}
 		else
 		{
-			FNavigationToolRemoveItemParams RemoveParams(InAddItemParams.Item);
+			FNavigationToolRemoveItemParams RemoveParams(ItemToAdd);
 			OldParent->RemoveChild(RemoveParams);
 		}
 	}
 
-	if (InAddItemParams.RelativeItem.IsValid()
+	if (ItemToAdd.IsValid()
 		&& InAddItemParams.RelativeDropZone.IsSet()
 		&& InAddItemParams.RelativeDropZone != EItemDropZone::OntoItem)
 	{
-		const int32 RelativeItemIndex = Children.Find(InAddItemParams.RelativeItem);
+		const int32 RelativeItemIndex = WeakChildren.Find(InAddItemParams.WeakRelativeItem);
 		if (RelativeItemIndex != INDEX_NONE)
 		{
 			const int32 TargetIndex = InAddItemParams.RelativeDropZone == EItemDropZone::BelowItem
 				? RelativeItemIndex + 1
 				: RelativeItemIndex;
 
-			Children.EmplaceAt(TargetIndex, InAddItemParams.Item);
+			WeakChildren.EmplaceAt(TargetIndex, InAddItemParams.WeakItem);
 		}
 		else
 		{
-			Children.EmplaceAt(0, InAddItemParams.Item);
+			WeakChildren.EmplaceAt(0, InAddItemParams.WeakItem);
 		}
 	}
 	else
 	{
-		Children.EmplaceAt(0, InAddItemParams.Item);
+		WeakChildren.EmplaceAt(0, InAddItemParams.WeakItem);
 	}
 
-	InAddItemParams.Item->SetParent(SharedThis(this));
+	ItemToAdd->SetParent(SharedThis(this));
 }
 
 bool FNavigationToolItem::RemoveChildChecked(const FNavigationToolRemoveItemParams& InRemoveItemParams)
 {
-	InRemoveItemParams.Item->SetParent(nullptr);
-	return Children.Remove(InRemoveItemParams.Item) > 0;
+	if (const FNavigationToolViewModelPtr Item = InRemoveItemParams.WeakItem.Pin())
+	{
+		Item->SetParent(nullptr);
+	}
+	return WeakChildren.Remove(InRemoveItemParams.WeakItem) > 0;
 }
 
-void FNavigationToolItem::HandleNewSortableChildren(TArray<FNavigationToolItemPtr> InSortableChildren)
+void FNavigationToolItem::HandleNewSortableChildren(TArray<FNavigationToolViewModelWeakPtr> InWeakSortableChildren)
 {
-	InSortableChildren.Sort([this](const FNavigationToolItemPtr& InItemA, const FNavigationToolItemPtr& InItemB)
+	InWeakSortableChildren.Sort([this](const FNavigationToolViewModelWeakPtr& InWeakItemA, const FNavigationToolViewModelWeakPtr& InWeakItemB)
 		{
-			const FNavigationToolSaveState* const SaveStateA = InItemA->GetProviderSaveState();
+			const TViewModelPtr<INavigationToolItem> ItemA = InWeakItemA.Pin();
+			const FNavigationToolSaveState* const SaveStateA = ItemA->GetProviderSaveState();
 			if (!SaveStateA)
 			{
 				return false;
 			}
 
-			const FNavigationToolSaveState* const SaveStateB = InItemB->GetProviderSaveState();
+			const TViewModelPtr<INavigationToolItem> ItemB = InWeakItemB.Pin();
+			const FNavigationToolSaveState* const SaveStateB = ItemB->GetProviderSaveState();
 			if (!SaveStateB)
 			{
 				return false;
 			}
 
-			const FNavigationToolSerializedTreeNode* const NodeA = SaveStateA->SerializedTree.FindTreeNode(InItemA->MakeSerializedItem());
-			const FNavigationToolSerializedTreeNode* const NodeB = SaveStateB->SerializedTree.FindTreeNode(InItemB->MakeSerializedItem());
+			const FNavigationToolSerializedTreeNode* const NodeA = SaveStateA->SerializedTree.FindTreeNode(ItemA->MakeSerializedItem());
+			const FNavigationToolSerializedTreeNode* const NodeB = SaveStateB->SerializedTree.FindTreeNode(ItemB->MakeSerializedItem());
 
 			return FNavigationToolSerializedTree::CompareTreeItemOrder(NodeA, NodeB);
 		});
 
 	FNavigationToolAddItemParams AddItemParams;
-	for (const FNavigationToolItemPtr& NewChild : InSortableChildren)
+	for (const FNavigationToolViewModelWeakPtr& WeakNewChild : InWeakSortableChildren)
 	{
+		const FNavigationToolViewModelPtr NewChild = WeakNewChild.Pin();
+		if (!NewChild.IsValid())
+		{
+			continue;
+		}
+
 		const FNavigationToolSaveState* const SaveState = NewChild->GetProviderSaveState();
 		if (!SaveState)
 		{
 			continue;
 		}
 
-		AddItemParams.Item = NewChild;
+		AddItemParams.WeakItem = WeakNewChild;
 
 		const FNavigationToolSerializedTreeNode* const TreeNode = SaveState->SerializedTree.FindTreeNode(NewChild->MakeSerializedItem());
-		if (TreeNode && Children.IsValidIndex(TreeNode->GetLocalIndex()))
+		if (TreeNode && WeakChildren.IsValidIndex(TreeNode->GetLocalIndex()))
 		{
 			// Add Before the Child at Index, so this Item is at the specific Index
-			AddItemParams.RelativeItem = Children[TreeNode->GetLocalIndex()];
+			AddItemParams.WeakRelativeItem = WeakChildren[TreeNode->GetLocalIndex()];
 			AddItemParams.RelativeDropZone = EItemDropZone::AboveItem;
 		}
 		else
 		{
 			// Add After Last, so this Item is the last item in the List
-			AddItemParams.RelativeItem = Children.IsEmpty() ? nullptr : Children.Last();
+			AddItemParams.WeakRelativeItem = WeakChildren.IsEmpty() ? nullptr : WeakChildren.Last();
 			AddItemParams.RelativeDropZone = EItemDropZone::BelowItem;
 		}
 
diff --git a/Engine/Plugins/Experimental/SequenceNavigator/Source/SequenceNavigator/Private/Items/NavigationToolItemId.cpp b/Engine/Plugins/Experimental/SequenceNavigator/Source/SequenceNavigator/Private/Items/NavigationToolItemId.cpp
--- a/Engine/Plugins/Experimental/SequenceNavigator/Source/SequenceNavigator/Private/Items/NavigationToolItemId.cpp
+++ b/Engine/Plugins/Experimental/SequenceNavigator/Source/SequenceNavigator/Private/Items/NavigationToolItemId.cpp
@@ -1,7 +1,8 @@
 ﻿// Copyright Epic Games, Inc. All Rights Reserved.
 
 #include "Items/NavigationToolItemId.h"
-#include "Items/INavigationToolItem.h"
+#include "ItemProxies/INavigationToolItemProxyFactory.h"
+#include "Items/NavigationToolItem.h"
 #include "MovieScene.h"
 #include "MovieSceneSection.h"
 #include "MovieSceneSequence.h"
@@ -23,7 +24,7 @@ FNavigationToolItemId::FNavigationToolItemId(const FString& InUniqueId)
 }
 
 FNavigationToolItemId::FNavigationToolItemId(const UObject* InObject
-	, const FNavigationToolItemPtr& InReferencingItem
+	, const FNavigationToolViewModelPtr& InReferencingItem
 	, const FString& InReferencingId)
 	: FNavigationToolItemId(InReferencingItem->GetItemId())
 {
@@ -32,7 +33,7 @@ FNavigationToolItemId::FNavigationToolItemId(const UObject* InObject
 	CalculateTypeHash();
 }
 
-FNavigationToolItemId::FNavigationToolItemId(const FNavigationToolItemPtr& InParentItem
+FNavigationToolItemId::FNavigationToolItemId(const FNavigationToolViewModelPtr& InParentItem
 	, const INavigationToolItemProxyFactory& InItemProxyFactory)
 	: FNavigationToolItemId(InParentItem->GetItemId())
 {
@@ -40,15 +41,15 @@ FNavigationToolItemId::FNavigationToolItemId(const FNavigationToolItemPtr& InPar
 	CalculateTypeHash();
 }
 
-FNavigationToolItemId::FNavigationToolItemId(const FNavigationToolItemPtr& InParentItem
-	, const INavigationToolItem& InItemProxy)
+FNavigationToolItemId::FNavigationToolItemId(const FNavigationToolViewModelPtr& InParentItem
+	, const FNavigationToolViewModelPtr& InItemProxy)
 	: FNavigationToolItemId(InParentItem->GetItemId())
 {
-	AddSeparatedSegment(Id, InItemProxy.GetTypeId().ToString());
+	AddSeparatedSegment(Id, InItemProxy.AsModel()->GetTypeTable().GetTypeName().ToString());
 	CalculateTypeHash();
 }
 
-FNavigationToolItemId::FNavigationToolItemId(const FNavigationToolItemPtr& InParentItem
+FNavigationToolItemId::FNavigationToolItemId(const FNavigationToolViewModelPtr& InParentItem
 	, const UMovieSceneSequence* const InSequence
 	, const UMovieSceneSection* const InSection
 	, const int32 InSectionIndex
@@ -60,7 +61,7 @@ FNavigationToolItemId::FNavigationToolItemId(const FNavigationToolItemPtr& InPar
 
 FNavigationToolItemId::FNavigationToolItemId(const UE::Sequencer::FViewModelPtr& InViewModel)
 {
-	using namespace UE::Sequencer;
+	using namespace Sequencer;
 
 	check(InViewModel.IsValid());
 
@@ -100,7 +101,7 @@ FNavigationToolItemId::FNavigationToolItemId(const UE::Sequencer::FViewModelPtr&
 	CalculateTypeHash();
 }
 
-void FNavigationToolItemId::ConstructId(const FNavigationToolItemPtr& InParentItem
+void FNavigationToolItemId::ConstructId(const FNavigationToolViewModelPtr& InParentItem
 	, const UMovieSceneSequence* const InSequence
 	, const UMovieSceneSection* const InSection
 	, const int32 InSectionIndex
diff --git a/Engine/Plugins/Experimental/SequenceNavigator/Source/SequenceNavigator/Private/Items/NavigationToolItemProxy.cpp b/Engine/Plugins/Experimental/SequenceNavigator/Source/SequenceNavigator/Private/Items/NavigationToolItemProxy.cpp
--- a/Engine/Plugins/Experimental/SequenceNavigator/Source/SequenceNavigator/Private/Items/NavigationToolItemProxy.cpp
+++ b/Engine/Plugins/Experimental/SequenceNavigator/Source/SequenceNavigator/Private/Items/NavigationToolItemProxy.cpp
@@ -8,7 +8,11 @@
 namespace UE::SequenceNavigator
 {
 
-FNavigationToolItemProxy::FNavigationToolItemProxy(INavigationTool& InTool, const FNavigationToolItemPtr& InParentItem)
+using namespace Sequencer;
+
+UE_SEQUENCER_DEFINE_CASTABLE(FNavigationToolItemProxy)
+
+FNavigationToolItemProxy::FNavigationToolItemProxy(INavigationTool& InTool, const FNavigationToolViewModelPtr& InParentItem)
 	: FNavigationToolItem(InTool, InParentItem)
 {
 	WeakParent = InParentItem;
@@ -17,26 +21,26 @@ FNavigationToolItemProxy::FNavigationToolItemProxy(INavigationTool& InTool, cons
 bool FNavigationToolItemProxy::IsItemValid() const
 {
 	return FNavigationToolItem::IsItemValid()
-		&& WeakParent.IsValid()
+		&& WeakParent.Pin().IsValid()
 		&& Tool.FindItem(WeakParent.Pin()->GetItemId()).IsValid();
 }
 
-void FNavigationToolItemProxy::FindChildren(TArray<FNavigationToolItemPtr>& OutChildren, bool bRecursive)
+void FNavigationToolItemProxy::FindChildren(TArray<FNavigationToolViewModelWeakPtr>& OutWeakChildren, const bool bInRecursive)
 {
-	const FNavigationToolItemPtr Parent = GetParent();
+	const FNavigationToolViewModelPtr Parent = GetParent();
 	if (!Parent.IsValid() || !Parent->IsAllowedInTool())
 	{
 		return;
 	}
-	
-	Super::FindChildren(OutChildren, bRecursive);
-	
-	GetProxiedItems(Parent.ToSharedRef(), OutChildren, bRecursive);
+
+	FNavigationToolItem::FindChildren(OutWeakChildren, bInRecursive);
+
+	GetProxiedItems(Parent, OutWeakChildren, bInRecursive);
 }
 
-void FNavigationToolItemProxy::SetParent(FNavigationToolItemPtr InParent)
+void FNavigationToolItemProxy::SetParent(FNavigationToolViewModelPtr InParent)
 {
-	Super::SetParent(InParent);
+	FNavigationToolItem::SetParent(InParent);
 
 	// Recalculate our item Id because we rely on what our parent is for our Id
 	RecalculateItemId();
@@ -45,7 +49,7 @@ void FNavigationToolItemProxy::SetParent(FNavigationToolItemPtr InParent)
 ENavigationToolItemViewMode FNavigationToolItemProxy::GetSupportedViewModes(const INavigationToolView& InToolView) const
 {
 	// Hide proxies if it has no children
-	if (Children.IsEmpty())
+	if (WeakChildren.IsEmpty())
 	{
 		return ENavigationToolItemViewMode::None;
 	}
@@ -59,9 +63,10 @@ bool FNavigationToolItemProxy::CanAutoExpand() const
 
 FNavigationToolItemId FNavigationToolItemProxy::CalculateItemId() const
 {
-	if (WeakParent.IsValid())
+	if (const FNavigationToolViewModelPtr Parent = WeakParent.Pin())
 	{
-		return FNavigationToolItemId(WeakParent.Pin(), *this);	
+		const FNavigationToolViewModelPtr ViewModel = AsItemViewModelConst();
+		return FNavigationToolItemId(Parent, ViewModel);	
 	}
 	return FNavigationToolItemId();
 }
diff --git a/Engine/Plugins/Experimental/SequenceNavigator/Source/SequenceNavigator/Private/Items/NavigationToolItemUtils.cpp b/Engine/Plugins/Experimental/SequenceNavigator/Source/SequenceNavigator/Private/Items/NavigationToolItemUtils.cpp
--- a/Engine/Plugins/Experimental/SequenceNavigator/Source/SequenceNavigator/Private/Items/NavigationToolItemUtils.cpp
+++ b/Engine/Plugins/Experimental/SequenceNavigator/Source/SequenceNavigator/Private/Items/NavigationToolItemUtils.cpp
@@ -7,7 +7,6 @@
 #include "Items/NavigationToolTreeRoot.h"
 #include "MovieSceneSequence.h"
 #include "NavigationTool.h"
-#include "NavigationToolView.h"
 #include "ScopedTransaction.h"
 #include "Sections/MovieSceneCinematicShotSection.h"
 #include "Sections/MovieSceneSubSection.h"
@@ -16,10 +15,12 @@
 
 #define LOCTEXT_NAMESPACE "NavigationToolItemUtils"
 
-namespace UE::SequenceNavigator
+namespace UE::SequenceNavigator::ItemUtils
 {
 
-bool CompareToolItemOrder(const FNavigationToolItemPtr& InA, const FNavigationToolItemPtr& InB)
+using namespace Sequencer;
+
+bool CompareToolItemOrder(const FNavigationToolViewModelPtr& InA, const FNavigationToolViewModelPtr& InB)
 {
 	if (!InA.IsValid())
 	{
@@ -31,10 +32,10 @@ bool CompareToolItemOrder(const FNavigationToolItemPtr& InA, const FNavigationTo
 		return false;
 	}
 
-	if (const FNavigationToolItemPtr LowestCommonAncestor = FNavigationTool::FindLowestCommonAncestor({ InA, InB }))
+	if (const FNavigationToolViewModelPtr LowestCommonAncestor = FNavigationTool::FindLowestCommonAncestor({ InA, InB }))
 	{
-		const TArray<FNavigationToolItemPtr> PathToA = LowestCommonAncestor->FindPath({ InA });
-		const TArray<FNavigationToolItemPtr> PathToB = LowestCommonAncestor->FindPath({ InB });
+		const TArray<FNavigationToolViewModelPtr> PathToA = LowestCommonAncestor->FindPath({ InA });
+		const TArray<FNavigationToolViewModelPtr> PathToB = LowestCommonAncestor->FindPath({ InB });
 
 		int32 Index = 0;
 	
@@ -63,16 +64,17 @@ bool CompareToolItemOrder(const FNavigationToolItemPtr& InA, const FNavigationTo
 	return false;
 }
 
-void SplitSortableAndUnsortableItems(const TArray<FNavigationToolItemPtr>& InItems
-	, TArray<FNavigationToolItemPtr>& OutSortable
-	, TArray<FNavigationToolItemPtr>& OutUnsortable)
+void SplitSortableAndUnsortableItems(const TArray<FNavigationToolViewModelWeakPtr>& InWeakItems
+	, TArray<FNavigationToolViewModelWeakPtr>& OutWeakSortable
+	, TArray<FNavigationToolViewModelWeakPtr>& OutWeakUnsortable)
 {
 	// Allocate both for worst case scenarios
-	OutSortable.Reserve(InItems.Num());
-	OutUnsortable.Reserve(InItems.Num());
+	OutWeakSortable.Reserve(InWeakItems.Num());
+	OutWeakUnsortable.Reserve(InWeakItems.Num());
 
-	for (const FNavigationToolItemPtr& Item : InItems)
+	for (const FNavigationToolViewModelWeakPtr& WeakItem : InWeakItems)
 	{
+		const FNavigationToolViewModelPtr Item = WeakItem.Pin();
 		if (!Item.IsValid())
 		{
 			continue;
@@ -80,27 +82,27 @@ void SplitSortableAndUnsortableItems(const TArray<FNavigationToolItemPtr>& InIte
 
 		if (Item->ShouldSort())
 		{
-			OutSortable.Add(Item);
+			OutWeakSortable.Add(WeakItem);
 		}
 		else
 		{
-			OutUnsortable.Add(Item);
+			OutWeakUnsortable.Add(WeakItem);
 		}
 	}
 }
 
-UMovieSceneSubSection* GetSequenceItemSubSection(const TSharedPtr<INavigationToolItem>& InItem)
+UMovieSceneSubSection* GetSequenceItemSubSection(const FNavigationToolViewModelPtr& InItem)
 {
-	if (const FNavigationToolSequence* const SequenceItem = InItem->CastTo<FNavigationToolSequence>())
+	if (const TViewModelPtr<FNavigationToolSequence> SequenceItem = InItem.ImplicitCast())
 	{
 		return SequenceItem->GetSubSection();
 	}
 	return nullptr;
 }
 
-UMovieSceneMetaData* GetSequenceItemMetaData(const TSharedPtr<INavigationToolItem>& InItem)
+UMovieSceneMetaData* GetSequenceItemMetaData(const FNavigationToolViewModelPtr& InItem)
 {
-	if (const FNavigationToolSequence* const SequenceItem = InItem->CastTo<FNavigationToolSequence>())
+	if (const TViewModelPtr<FNavigationToolSequence> SequenceItem = InItem.ImplicitCast())
 	{
 		return FSubTrackEditorUtil::FindOrAddMetaData(SequenceItem->GetSequence());
 	}
@@ -108,22 +110,27 @@ UMovieSceneMetaData* GetSequenceItemMetaData(const TSharedPtr<INavigationToolIte
 }
 
 void RemoveSequenceDisplayNameParentPrefix(FText& InOutDisplayName
-	, const TSharedRef<FNavigationToolSequence>& InSequenceItem)
+	, const TViewModelPtr<FNavigationToolSequence>& InSequenceItem)
 {
+	if (!InSequenceItem.IsValid())
+	{
+		return;
+	}
+
 	UMovieSceneSequence* const Sequence = InSequenceItem->GetSequence();
 	if (!Sequence)
 	{
 		return;
 	}
 
-	const FNavigationToolItemPtr Parent = InSequenceItem->GetParent();
-	if (!Parent.IsValid() || Parent->IsA<FNavigationToolTreeRoot>())
+	const FNavigationToolViewModelPtr Parent = InSequenceItem->GetParent();
+	if (!Parent.IsValid() || Parent.AsModel()->IsA<FNavigationToolTreeRoot>())
 	{
 		return;
 	}
 
 	// Go up the hierarchy to find the parent sequence
-	const FNavigationToolSequence* const ParentSequenceItem = InSequenceItem->FindAncestorOfType<FNavigationToolSequence>();
+	const TViewModelPtr<FNavigationToolSequence> ParentSequenceItem = InSequenceItem->FindAncestorOfType<FNavigationToolSequence>();
 	if (!ParentSequenceItem)
 	{
 		return;
@@ -143,7 +150,7 @@ void RemoveSequenceDisplayNameParentPrefix(FText& InOutDisplayName
 		return;
 	}
 
-	// Remove the parent prefix from child name
+	// Remove the parent prefix from the child name
 	static FString SeparatorString = TEXT("_");
 	const FString ParentSequenceNameStr = ParentSequence->GetDisplayName().ToString();
 
diff --git a/Engine/Plugins/Experimental/SequenceNavigator/Source/SequenceNavigator/Private/Items/NavigationToolMarker.cpp b/Engine/Plugins/Experimental/SequenceNavigator/Source/SequenceNavigator/Private/Items/NavigationToolMarker.cpp
--- a/Engine/Plugins/Experimental/SequenceNavigator/Source/SequenceNavigator/Private/Items/NavigationToolMarker.cpp
+++ b/Engine/Plugins/Experimental/SequenceNavigator/Source/SequenceNavigator/Private/Items/NavigationToolMarker.cpp
@@ -18,11 +18,15 @@
 namespace UE::SequenceNavigator
 {
 
+using namespace Sequencer;
+
+UE_SEQUENCER_DEFINE_CASTABLE(FNavigationToolMarker)
+
 FNavigationToolMarker::FNavigationToolMarker(INavigationTool& InTool
-	, const FNavigationToolItemPtr& InParentItem
+	, const FNavigationToolViewModelPtr& InParentItem
 	, const TSharedPtr<FNavigationToolSequence>& InParentSequenceItem
 	, const int32 InMarkedFrameIndex)
-	: Super(InTool, InParentItem)
+	: FNavigationToolItem(InTool, InParentItem)
 	, WeakParentSequenceItem(InParentSequenceItem)
 	, MarkedFrameIndex(InMarkedFrameIndex)
 {
@@ -30,7 +34,7 @@ FNavigationToolMarker::FNavigationToolMarker(INavigationTool& InTool
 
 bool FNavigationToolMarker::IsItemValid() const
 {
-	return WeakParentSequenceItem.IsValid() && MarkedFrameIndex != INDEX_NONE;
+	return WeakParentSequenceItem.Pin().IsValid() && MarkedFrameIndex != INDEX_NONE;
 }
 
 bool FNavigationToolMarker::IsAllowedInTool() const
@@ -105,7 +109,7 @@ void FNavigationToolMarker::Select(FNavigationToolScopedSelection& InSelection)
 
 void FNavigationToolMarker::OnSelect()
 {
-	Super::OnSelect();
+	FNavigationToolItem::OnSelect();
 }
 
 void FNavigationToolMarker::OnDoubleClick()
@@ -188,45 +192,46 @@ bool FNavigationToolMarker::CanRename() const
 	return MovieScene->CanModify();
 }
 
-bool FNavigationToolMarker::Rename(const FString& InName)
+void FNavigationToolMarker::Rename(const FText& InNewName)
 {
-	if (!WeakParentSequenceItem.IsValid() || MarkedFrameIndex == INDEX_NONE)
+	const TViewModelPtr<FNavigationToolSequence> ParentSequenceItem = WeakParentSequenceItem.ImplicitPin();
+	if (!ParentSequenceItem.IsValid() || MarkedFrameIndex == INDEX_NONE)
 	{
-		return false;
+		return;
 	}
 
 	UMovieScene* const MovieScene = WeakParentSequenceItem.Pin()->GetSequenceMovieScene();
 	if (!MovieScene)
 	{
-		return false;
+		return;
 	}
 
 	const TArray<FMovieSceneMarkedFrame>& MarkedFrames = MovieScene->GetMarkedFrames();
 	if (!MarkedFrames.IsValidIndex(MarkedFrameIndex))
 	{
-		return false;
+		return;
 	}
 
 	FMovieSceneMarkedFrame& MarkedFrame = const_cast<FMovieSceneMarkedFrame&>(MarkedFrames[MarkedFrameIndex]);
 
-	bool bRenamed = false;
+	const FString NewNameString = InNewName.ToString();
 
-	if (!InName.Equals(MarkedFrame.Label, ESearchCase::CaseSensitive))
+	if (NewNameString.Equals(MarkedFrame.Label, ESearchCase::CaseSensitive))
 	{
-		const FScopedTransaction Transaction(LOCTEXT("RenameTransaction", "Rename Marker"));
+		return;
+	}
 
-		MovieScene->Modify();
-		MarkedFrame.Label = InName;
+	const FScopedTransaction Transaction(LOCTEXT("RenameTransaction", "Rename Marker"));
 
-		bRenamed = true;
-	}
+	MovieScene->Modify();
 
-	return bRenamed;
+	MarkedFrame.Label = NewNameString;
 }
 
 FFrameNumber FNavigationToolMarker::GetInTime() const
 {
-	if (!WeakParentSequenceItem.IsValid() || MarkedFrameIndex == INDEX_NONE)
+	const TViewModelPtr<FNavigationToolSequence> ParentSequenceItem = WeakParentSequenceItem.ImplicitPin();
+	if (!ParentSequenceItem.IsValid() || MarkedFrameIndex == INDEX_NONE)
 	{
 		return 0;
 	}
@@ -256,7 +261,8 @@ FFrameNumber FNavigationToolMarker::GetInTime() const
 
 void FNavigationToolMarker::SetInTime(const FFrameNumber& InTime)
 {
-	if (!WeakParentSequenceItem.IsValid() || MarkedFrameIndex == INDEX_NONE)
+	const TViewModelPtr<FNavigationToolSequence> ParentSequenceItem = WeakParentSequenceItem.ImplicitPin();
+	if (!ParentSequenceItem.IsValid() || MarkedFrameIndex == INDEX_NONE)
 	{
 		return;
 	}
