[SequenceNavigator] Refactor to use Sequencer view models and type macros

diff --git a/Engine/Plugins/Experimental/SequenceNavigator/Source/SequenceNavigator/Private/Items/NavigationToolSequence.cpp b/Engine/Plugins/Experimental/SequenceNavigator/Source/SequenceNavigator/Private/Items/NavigationToolSequence.cpp
--- a/Engine/Plugins/Experimental/SequenceNavigator/Source/SequenceNavigator/Private/Items/NavigationToolSequence.cpp
+++ b/Engine/Plugins/Experimental/SequenceNavigator/Source/SequenceNavigator/Private/Items/NavigationToolSequence.cpp
@@ -14,6 +14,7 @@
 #include "MVVM/SectionModelStorageExtension.h"
 #include "MVVM/ViewModels/SectionModel.h"
 #include "NavigationToolScopedSelection.h"
+#include "NavigationToolSettings.h"
 #include "ScopedTransaction.h"
 #include "Sections/MovieSceneCinematicShotSection.h"
 #include "Sections/MovieSceneSubSection.h"
@@ -26,17 +27,19 @@
 
 #define LOCTEXT_NAMESPACE "NavigationToolSequence"
 
-using namespace UE::Sequencer;
-
 namespace UE::SequenceNavigator
 {
 
+using namespace Sequencer;
+
+UE_SEQUENCER_DEFINE_CASTABLE(FNavigationToolSequence)
+
 FNavigationToolSequence::FNavigationToolSequence(INavigationTool& InTool
-	, const FNavigationToolItemPtr& InParentItem
+	, const FNavigationToolViewModelPtr& InParentItem
 	, UMovieSceneSequence* const InSequence
 	, UMovieSceneSubSection* const InSubSection
 	, const int32 InSubSectionIndex)
-	: Super(InTool, InParentItem)
+	: FNavigationToolItem(InTool, InParentItem)
 	, IRevisionControlExtension(InSequence)
 	, WeakSubSection(InSubSection)
 	, SubSectionIndex(InSubSectionIndex)
@@ -59,9 +62,9 @@ bool FNavigationToolSequence::IsAllowedInTool() const
 	return WeakSequence.IsValid();
 }
 
-void FNavigationToolSequence::FindChildren(TArray<FNavigationToolItemPtr>& OutChildren, const bool bInRecursive)
+void FNavigationToolSequence::FindChildren(TArray<FNavigationToolViewModelWeakPtr>& OutWeakChildren, const bool bInRecursive)
 {
-	Super::FindChildren(OutChildren, bInRecursive);
+	FNavigationToolItem::FindChildren(OutWeakChildren, bInRecursive);
 
 	UMovieSceneSequence* const Sequence = GetSequence();
 	if (!Sequence)
@@ -75,35 +78,43 @@ void FNavigationToolSequence::FindChildren(TArray<FNavigationToolItemPtr>& OutCh
 		return;
 	}
 
-	const TSharedRef<FNavigationToolSequence> SharedThisRef = SharedThis(this);
 	const TSharedPtr<FNavigationToolProvider> Provider = GetProvider();
+	if (!Provider.IsValid())
+	{
+		return;
+	}
+
+	const TSharedRef<FNavigationToolProvider> ProviderRef = Provider.ToSharedRef();
+	const FNavigationToolViewModelPtr ParentItem = AsItemViewModel();
 	const TArray<UMovieSceneTrack*> MovieSceneTracks = MovieScene->GetTracks();
 
 	for (UMovieSceneTrack* const Track : MovieSceneTracks)
 	{
-		FNavigationToolItemPtr NewItem;
+		FNavigationToolViewModelPtr NewItem;
 
 		if (UMovieSceneSubTrack* const SubTrack = Cast<UMovieSceneSubTrack>(Track))
 		{
-			NewItem = Tool.FindOrAdd<FNavigationToolSubTrack>(GetProvider()
-				, SharedThisRef, SubTrack, Sequence, WeakSubSection, SubSectionIndex);
+			NewItem = Tool.FindOrAdd<FNavigationToolSubTrack>(ProviderRef
+				, ParentItem, SubTrack, Sequence, WeakSubSection, SubSectionIndex);
 		}
 		else if (Track)
 		{
-			NewItem = Tool.FindOrAdd<FNavigationToolTrack>(GetProvider()
-				, SharedThisRef, Track, Sequence, WeakSubSection, SubSectionIndex);
+			NewItem = Tool.FindOrAdd<FNavigationToolTrack>(ProviderRef
+				, ParentItem, Track, Sequence, WeakSubSection, SubSectionIndex);
 		}
 
 		if (NewItem.IsValid())
 		{
-			OutChildren.Add(NewItem);
+			OutWeakChildren.Add(NewItem);
 			if (bInRecursive)
 			{
-				NewItem->FindChildren(OutChildren, bInRecursive);
+				NewItem->FindChildren(OutWeakChildren, bInRecursive);
 			}
 		}
 	}
 
+	const TSharedRef<FNavigationToolSequence> ThisRef = SharedThis(this);
+
 	// Only show actor bindings for a Sequence
 	const TArray<FMovieSceneBinding> Bindings = GetSortedBindings();
 	for (const FMovieSceneBinding& Binding : Bindings)
@@ -112,12 +123,12 @@ void FNavigationToolSequence::FindChildren(TArray<FNavigationToolItemPtr>& OutCh
 		const UClass* const BoundObjectClass = MovieSceneHelpers::GetBoundObjectClass(Sequence, BindingObjectGuid);
 		if (BoundObjectClass && BoundObjectClass->IsChildOf<AActor>())
 		{
-			const FNavigationToolItemPtr NewItem = Tool.FindOrAdd<FNavigationToolActor>(GetProvider()
-				, SharedThisRef, SharedThisRef, Binding);
-			OutChildren.Add(NewItem);
+			const FNavigationToolViewModelPtr NewItem = Tool.FindOrAdd<FNavigationToolActor>(ProviderRef
+				, ThisRef, ThisRef, Binding);
+			OutWeakChildren.Add(NewItem);
 			if (bInRecursive)
 			{
-				NewItem->FindChildren(OutChildren, bInRecursive);
+				NewItem->FindChildren(OutWeakChildren, bInRecursive);
 			}
 		}
 	}
@@ -126,31 +137,31 @@ void FNavigationToolSequence::FindChildren(TArray<FNavigationToolItemPtr>& OutCh
 	const TArray<FMovieSceneMarkedFrame>& MarkedFrames = MovieScene->GetMarkedFrames();
 	for (int32 Index = 0; Index < MarkedFrames.Num(); ++Index)
 	{
-		const FNavigationToolItemPtr NewItem = Tool.FindOrAdd<FNavigationToolMarker>(GetProvider()
-			, SharedThisRef, SharedThisRef, Index);
-		OutChildren.Add(NewItem);
+		const FNavigationToolViewModelPtr NewItem = Tool.FindOrAdd<FNavigationToolMarker>(ProviderRef
+			, ThisRef, ThisRef, Index);
+		OutWeakChildren.Add(NewItem);
 		if (bInRecursive)
 		{
-			NewItem->FindChildren(OutChildren, bInRecursive);
+			NewItem->FindChildren(OutWeakChildren, bInRecursive);
 		}
 	}
 }
 
 void FNavigationToolSequence::GetItemProxies(TArray<TSharedPtr<FNavigationToolItemProxy>>& OutItemProxies)
 {
-	Super::GetItemProxies(OutItemProxies);
+	FNavigationToolItem::GetItemProxies(OutItemProxies);
 }
 
 bool FNavigationToolSequence::AddChild(const FNavigationToolAddItemParams& InAddItemParams)
 {
 	// @TODO: handle sequence being moved to this item
-	return Super::AddChild(InAddItemParams);
+	return FNavigationToolItem::AddChild(InAddItemParams);
 }
 
 bool FNavigationToolSequence::RemoveChild(const FNavigationToolRemoveItemParams& InRemoveItemParams)
 {
 	/// @TODO: handle sequence being moved to this item
-	return Super::RemoveChild(InRemoveItemParams);
+	return FNavigationToolItem::RemoveChild(InRemoveItemParams);
 }
 
 ENavigationToolItemViewMode FNavigationToolSequence::GetSupportedViewModes(const INavigationToolView& InToolView) const
@@ -162,6 +173,8 @@ ENavigationToolItemViewMode FNavigationToolSequence::GetSupportedViewModes(const
 
 FText FNavigationToolSequence::GetDisplayName() const
 {
+	using namespace ItemUtils;
+
 	FText NewDisplayName = FText::GetEmpty();
 
 	if (const UMovieSceneCinematicShotSection* const ShotSection = Cast<UMovieSceneCinematicShotSection>(WeakSubSection.Get()))
@@ -181,16 +194,15 @@ FText FNavigationToolSequence::GetDisplayName() const
 		return NewDisplayName;
 	}
 
-	// Apply additional options to display name
-	const TSharedRef<FNavigationToolSequence> ThisSequenceRef = SharedThis(const_cast<FNavigationToolSequence*>(this));
+	// Apply additional options to the display name
 	const UNavigationToolSettings* const ToolSettings = GetDefault<UNavigationToolSettings>();
 
 	if (ToolSettings->ShouldUseShortNames())
 	{
-		RemoveSequenceDisplayNameParentPrefix(NewDisplayName, ThisSequenceRef);
+		RemoveSequenceDisplayNameParentPrefix(NewDisplayName, AsItemViewModelConst().ImplicitCast());
 	}
 
-	if (UMovieSceneSequence* const Sequence = ThisSequenceRef->GetSequence())
+	if (UMovieSceneSequence* const Sequence = GetSequence())
 	{
 		AppendSequenceDisplayNameDirtyStatus(NewDisplayName, *Sequence);
 	}
@@ -232,7 +244,7 @@ void FNavigationToolSequence::Select(FNavigationToolScopedSelection& InSelection
 
 void FNavigationToolSequence::OnSelect()
 {
-	Super::OnSelect();
+	FNavigationToolItem::OnSelect();
 }
 
 void FNavigationToolSequence::OnDoubleClick()
@@ -247,8 +259,8 @@ void FNavigationToolSequence::OnDoubleClick()
 
 	if (ModifierKeys.IsAltDown())
 	{
-		// Instead of focusing the sequence in the Sequencer, focus the parent sequence and select sequence
-		if (const FNavigationToolSequence* const ParentSequenceItem = FindAncestorOfType<FNavigationToolSequence>())
+		// Instead of focusing the sequence in the sequencer, focus the parent sequence and select the sequence
+		if (const TViewModelPtr<FNavigationToolSequence> ParentSequenceItem = FindAncestorOfType<FNavigationToolSequence>())
 		{
 			if (UMovieSceneSequence* const ParentSequence = ParentSequenceItem->GetSequence())
 			{
@@ -270,7 +282,7 @@ void FNavigationToolSequence::OnObjectsReplaced(const TMap<UObject*, UObject*>&
 	}
 
 	// This handles calling OnObjectsReplaced for every child item
-	Super::OnObjectsReplaced(InReplacementMap, bInRecursive);
+	FNavigationToolItem::OnObjectsReplaced(InReplacementMap, bInRecursive);
 }
 
 FNavigationToolItemId FNavigationToolSequence::CalculateItemId() const
@@ -297,81 +309,120 @@ bool FNavigationToolSequence::CanRename() const
 	return !MovieScene->IsReadOnly();*/
 }
 
-bool FNavigationToolSequence::Rename(const FString& InName)
+void FNavigationToolSequence::Rename(const FText& InNewName)
 {
 	UMovieSceneSequence* const Sequence = GetSequence();
 	if (!Sequence)
 	{
-		return false;
+		return;
 	}
 
-	bool bRenamed = false;
+	const FString NewNameString = InNewName.ToString();
 
-	if (!InName.Equals(Sequence->GetName(), ESearchCase::CaseSensitive))
+	if (NewNameString.Equals(Sequence->GetName(), ESearchCase::CaseSensitive))
 	{
-		//const FScopedTransaction Transaction(LOCTEXT("SequenceRename", "Rename Sequence"));
+		return;
+	}
 
-		if (UMovieSceneCinematicShotSection* const ShotSection = Cast<UMovieSceneCinematicShotSection>(WeakSubSection.Get()))
-		{
-			ShotSection->Modify();
-			ShotSection->SetShotDisplayName(InName);
+	const FScopedTransaction Transaction(LOCTEXT("SequenceRename", "Rename Sequence"));
 
-			Tool.NotifyToolItemRenamed(SharedThis(this));
+	if (UMovieSceneCinematicShotSection* const ShotSection = Cast<UMovieSceneCinematicShotSection>(WeakSubSection.Get()))
+	{
+		ShotSection->Modify();
+		ShotSection->SetShotDisplayName(NewNameString);
 
-			bRenamed = true;
-		}
-		else
+		Tool.NotifyToolItemRenamed(SharedThis(this));
+	}
+	else
+	{
+		Sequence->Modify();
+
+		if (Sequence->Rename(*NewNameString))
 		{
-			Sequence->Modify();
+			Tool.NotifyToolItemRenamed(SharedThis(this));
+		}
+	}
+}
 
-			if (Sequence->Rename(*InName))
+bool FNavigationToolSequence::IsDeactivated() const
+{
+	using namespace ItemUtils;
+
+	if (WeakSubSection.IsValid())
+	{
+		return !WeakSubSection->IsActive();
+	}
+
+	const ENavigationToolCompareState State = CompareChildrenItemState<IDeactivatableExtension>(AsItemViewModelConst(),
+		[](const TViewModelPtr<IDeactivatableExtension>& InItem)
+			{
+				return InItem->IsDeactivated();
+			},
+		[](const TViewModelPtr<IDeactivatableExtension>& InItem)
 			{
-				Tool.NotifyToolItemRenamed(SharedThis(this));
+				return !InItem->IsDeactivated();
+			});
 
-				bRenamed = true;
-			}
+	return State == ENavigationToolCompareState::AllTrue;
+}
+
+void FNavigationToolSequence::SetIsDeactivated(const bool bInIsDeactivated)
+{
+	const bool bNewActiveState = !bInIsDeactivated;
+
+	if (WeakSubSection.IsValid())
+	{
+		if (!WeakSubSection->IsReadOnly() && WeakSubSection->IsActive() != bNewActiveState)
+		{
+			WeakSubSection->Modify();
+			WeakSubSection->SetIsActive(bNewActiveState);
 		}
 	}
 
-	return bRenamed;
+	for (const TViewModelPtr<IDeactivatableExtension>& InactivableItem : GetChildrenOfType<IDeactivatableExtension>())
+	{
+		InactivableItem->SetIsDeactivated(bInIsDeactivated);
+	}
 }
 
-EItemSequenceInactiveState FNavigationToolSequence::GetInactiveState() const
+ELockableLockState FNavigationToolSequence::GetLockState() const
 {
+	using namespace ItemUtils;
+
 	if (WeakSubSection.IsValid())
 	{
-		return WeakSubSection->IsActive() ? EItemSequenceInactiveState::None : EItemSequenceInactiveState::Inactive;
+		return WeakSubSection->IsLocked() ? ELockableLockState::Locked : ELockableLockState::None;
 	}
 
-	const ENavigationToolCompareState State = CompareChildrenItemState<ISequenceInactivableExtension>(this,
-		[](const ISequenceInactivableExtension* const InItem)
+	const ENavigationToolCompareState State = CompareChildrenItemState<ILockableExtension>(AsItemViewModelConst(),
+		[](const TViewModelPtr<ILockableExtension>& InItem)
 			{
-				return InItem->GetInactiveState() == EItemSequenceInactiveState::Inactive;
+				return InItem->GetLockState() == ELockableLockState::Locked;
 			},
-		[](const ISequenceInactivableExtension* const InItem)
+		[](const TViewModelPtr<ILockableExtension>& InItem)
 			{
-				return InItem->GetInactiveState() == EItemSequenceInactiveState::None;
+				return InItem->GetLockState() != ELockableLockState::Locked;
 			});
 
-	return static_cast<EItemSequenceInactiveState>(State);
+	return static_cast<ELockableLockState>(State);
 }
 
-void FNavigationToolSequence::SetIsInactive(const bool bInIsInactive)
+void FNavigationToolSequence::SetIsLocked(const bool bInIsLocked)
 {
-	const bool bNewActiveState = !bInIsInactive;
-
 	if (WeakSubSection.IsValid())
 	{
-		if (!WeakSubSection->IsReadOnly() && WeakSubSection->IsActive() != bNewActiveState)
+		if (WeakSubSection->IsLocked() != bInIsLocked)
 		{
 			WeakSubSection->Modify();
-			WeakSubSection->SetIsActive(bNewActiveState);
+			WeakSubSection->SetIsLocked(bInIsLocked);
 		}
+
+		return;
 	}
 
-	for (ISequenceInactivableExtension* const InactivableItem : GetChildrenOfType<ISequenceInactivableExtension>())
+	for (const TViewModelPtr<ILockableExtension>& LockableItem : GetChildrenOfType<ILockableExtension>())
 	{
-		InactivableItem->SetIsInactive(bInIsInactive);
+		LockableItem->SetIsLocked(bInIsLocked);
 	}
 }
 
@@ -403,51 +454,12 @@ void FNavigationToolSequence::SetMarkerVisibility(const bool bInVisible)
 
 	ShowGloballyMarkedFramesForSequence(*Sequencer, Sequence, bInVisible);
 
-	for (IMarkerVisibilityExtension* const MarkerVisibilityItem : GetChildrenOfType<IMarkerVisibilityExtension>())
+	for (const TViewModelPtr<IMarkerVisibilityExtension>& MarkerVisibilityItem : GetChildrenOfType<IMarkerVisibilityExtension>())
 	{
 		MarkerVisibilityItem->SetMarkerVisibility(bInVisible);
 	}
 }
 
-EItemSequenceLockState FNavigationToolSequence::GetLockState() const
-{
-	if (WeakSubSection.IsValid())
-	{
-		return WeakSubSection->IsLocked() ? EItemSequenceLockState::Locked : EItemSequenceLockState::None;
-	}
-
-	const ENavigationToolCompareState State = CompareChildrenItemState<ISequenceLockableExtension>(this,
-		[](const ISequenceLockableExtension* const InItem)
-			{
-				return InItem->GetLockState() == EItemSequenceLockState::Locked;
-			},
-		[](const ISequenceLockableExtension* const InItem)
-			{
-				return InItem->GetLockState() != EItemSequenceLockState::Locked;
-			});
-
-	return static_cast<EItemSequenceLockState>(State);
-}
-
-void FNavigationToolSequence::SetIsLocked(const bool bInIsLocked)
-{
-	if (WeakSubSection.IsValid())
-	{
-		if (WeakSubSection->IsLocked() != bInIsLocked)
-		{
-			WeakSubSection->Modify();
-			WeakSubSection->SetIsLocked(bInIsLocked);
-		}
-
-		return;
-	}
-
-	for (ISequenceLockableExtension* const LockableItem : GetChildrenOfType<ISequenceLockableExtension>())
-	{
-		LockableItem->SetIsLocked(bInIsLocked);
-	}
-}
-
 TOptional<FColor> FNavigationToolSequence::GetColor() const
 {
 	if (!WeakSubSection.IsValid())
@@ -526,13 +538,13 @@ EItemContainsPlayhead FNavigationToolSequence::ContainsPlayhead() const
 	}
 
 	const UMovieSceneSequence* const Sequence = WeakSequence.Get();
-	if (!IsValid(Sequence))
+	if (!Sequence)
 	{
 		return EItemContainsPlayhead::None;
 	}
 
 	UMovieSceneSequence* const FocusedSequence = Sequencer->GetFocusedMovieSceneSequence();
-	if (!IsValid(FocusedSequence))
+	if (!FocusedSequence)
 	{
 		return EItemContainsPlayhead::None;
 	}
@@ -541,8 +553,7 @@ EItemContainsPlayhead FNavigationToolSequence::ContainsPlayhead() const
 
 	if (Sequence == FocusedSequence)
 	{
-		UMovieScene* const SequenceMovieScene = Sequence->GetMovieScene();
-		if (IsValid(SequenceMovieScene))
+		if (UMovieScene* const SequenceMovieScene = Sequence->GetMovieScene())
 		{
 			return SequenceMovieScene->GetPlaybackRange().Contains(PlayheadTime.Time.FrameNumber)
 				? EItemContainsPlayhead::ContainsPlayhead : EItemContainsPlayhead::None;
@@ -550,8 +561,7 @@ EItemContainsPlayhead FNavigationToolSequence::ContainsPlayhead() const
 	}
 	else
 	{
-		UMovieSceneSubSection* const SubSection = GetSubSection();
-		if (IsValid(SubSection))
+		if (UMovieSceneSubSection* const SubSection = GetSubSection())
 		{
 			return SubSection->GetTrueRange().Contains(PlayheadTime.Time.FrameNumber)
 				? EItemContainsPlayhead::ContainsPlayhead : EItemContainsPlayhead::None;
diff --git a/Engine/Plugins/Experimental/SequenceNavigator/Source/SequenceNavigator/Private/Items/NavigationToolSubTrack.cpp b/Engine/Plugins/Experimental/SequenceNavigator/Source/SequenceNavigator/Private/Items/NavigationToolSubTrack.cpp
--- a/Engine/Plugins/Experimental/SequenceNavigator/Source/SequenceNavigator/Private/Items/NavigationToolSubTrack.cpp
+++ b/Engine/Plugins/Experimental/SequenceNavigator/Source/SequenceNavigator/Private/Items/NavigationToolSubTrack.cpp
@@ -16,15 +16,17 @@
 namespace UE::SequenceNavigator
 {
 
-using namespace UE::Sequencer;
+using namespace Sequencer;
+
+UE_SEQUENCER_DEFINE_CASTABLE(FNavigationToolSubTrack)
 
 FNavigationToolSubTrack::FNavigationToolSubTrack(INavigationTool& InTool
-	, const FNavigationToolItemPtr& InParentItem
+	, const FNavigationToolViewModelPtr& InParentItem
 	, UMovieSceneSubTrack* const InSubTrack
 	, const TWeakObjectPtr<UMovieSceneSequence>& InSequence
 	, const TWeakObjectPtr<UMovieSceneSection>& InSection
 	, const int32 InSubSectionIndex)
-	: Super(InTool
+	: FNavigationToolTrack(InTool
 	, InParentItem
 	, InSubTrack
 	, InSequence
@@ -34,29 +36,38 @@ FNavigationToolSubTrack::FNavigationToolSubTrack(INavigationTool& InTool
 	OnTrackObjectChanged();
 }
 
-void FNavigationToolSubTrack::FindChildren(TArray<FNavigationToolItemPtr>& OutChildren, const bool bInRecursive)
+void FNavigationToolSubTrack::FindChildren(TArray<FNavigationToolViewModelWeakPtr>& OutWeakChildren, const bool bInRecursive)
 {
-	Super::FindChildren(OutChildren, bInRecursive);
+	FNavigationToolTrack::FindChildren(OutWeakChildren, bInRecursive);
 
-	if (UMovieSceneSubTrack* const SubTrack = GetSubTrack())
+	UMovieSceneSubTrack* const SubTrack = GetSubTrack();
+	if (!SubTrack)
 	{
-		const TSharedRef<FNavigationToolSubTrack> SharedThisRef = SharedThis(this);
-		const TSharedPtr<FNavigationToolProvider> Provider = GetProvider();
-		const TArray<UMovieSceneSection*>& AllSections = SubTrack->GetAllSections();
+		return;
+	}
 
-		for (int32 Index = 0; Index < AllSections.Num(); ++Index)
+	const TSharedPtr<FNavigationToolProvider> Provider = GetProvider();
+	if (!Provider.IsValid())
+	{
+		return;
+	}
+
+	const TSharedRef<FNavigationToolProvider> ProviderRef = Provider.ToSharedRef();
+	const FNavigationToolViewModelPtr ParentItem = AsItemViewModel();
+	const TArray<UMovieSceneSection*>& AllSections = SubTrack->GetAllSections();
+
+	for (int32 Index = 0; Index < AllSections.Num(); ++Index)
+	{
+		if (UMovieSceneSubSection* const SubSection = Cast<UMovieSceneSubSection>(AllSections[Index]))
 		{
-			if (UMovieSceneSubSection* const SubSection = Cast<UMovieSceneSubSection>(AllSections[Index]))
+			if (UMovieSceneSequence* const Sequence = SubSection->GetSequence())
 			{
-				if (UMovieSceneSequence* const Sequence = SubSection->GetSequence())
+				const FNavigationToolViewModelPtr NewItem = Tool.FindOrAdd<FNavigationToolSequence>(ProviderRef
+					, ParentItem, Sequence, SubSection, Index);
+				OutWeakChildren.Add(NewItem);
+				if (bInRecursive)
 				{
-					const FNavigationToolItemPtr NewItem = Tool.FindOrAdd<FNavigationToolSequence>(Provider
-						, SharedThisRef, Sequence, SubSection, Index);
-					OutChildren.Add(NewItem);
-					if (bInRecursive)
-					{
-						NewItem->FindChildren(OutChildren, bInRecursive);
-					}
+					NewItem->FindChildren(OutWeakChildren, bInRecursive);
 				}
 			}
 		}
@@ -65,12 +76,12 @@ void FNavigationToolSubTrack::FindChildren(TArray<FNavigationToolItemPtr>& OutCh
 
 TOptional<EItemDropZone> FNavigationToolSubTrack::CanAcceptDrop(const FDragDropEvent& InDragDropEvent, EItemDropZone InDropZone)
 {
-	return Super::CanAcceptDrop(InDragDropEvent, InDropZone);
+	return FNavigationToolTrack::CanAcceptDrop(InDragDropEvent, InDropZone);
 }
 
 FReply FNavigationToolSubTrack::AcceptDrop(const FDragDropEvent& InDragDropEvent, EItemDropZone InDropZone)
 {
-	return Super::AcceptDrop(InDragDropEvent, InDropZone);
+	return FNavigationToolTrack::AcceptDrop(InDragDropEvent, InDropZone);
 }
 
 UMovieSceneSubTrack* FNavigationToolSubTrack::GetSubTrack() const
@@ -78,29 +89,31 @@ UMovieSceneSubTrack* FNavigationToolSubTrack::GetSubTrack() const
 	return Cast<UMovieSceneSubTrack>(GetTrack());
 }
 
-EItemSequenceInactiveState FNavigationToolSubTrack::GetInactiveState() const
+bool FNavigationToolSubTrack::IsDeactivated() const
 {
+	using namespace ItemUtils;
+
 	if (UMovieSceneSubSection* const ThisSubSection = Cast<UMovieSceneSubSection>(WeakSection.Get()))
 	{
-		return ThisSubSection->IsActive() ? EItemSequenceInactiveState::None : EItemSequenceInactiveState::Inactive;
+		return !ThisSubSection->IsActive();
 	}
 
-	const ENavigationToolCompareState State = CompareChildrenItemState<ISequenceInactivableExtension>(this,
-		[](const ISequenceInactivableExtension* const InInactivableItem)
+	const ENavigationToolCompareState State = CompareChildrenItemState<IDeactivatableExtension>(AsItemViewModelConst(),
+		[](const TViewModelPtr<IDeactivatableExtension>& InItem)
 			{
-				return InInactivableItem->GetInactiveState() == EItemSequenceInactiveState::Inactive;
+				return InItem->IsDeactivated();
 			},
-		[](const ISequenceInactivableExtension* const InInactivableItem)
+		[](const TViewModelPtr<IDeactivatableExtension>& InItem)
 			{
-				return InInactivableItem->GetInactiveState() == EItemSequenceInactiveState::None;
+				return !InItem->IsDeactivated();
 			});
 
-	return static_cast<EItemSequenceInactiveState>(State);
+	return State == ENavigationToolCompareState::AllTrue;
 }
 
-void FNavigationToolSubTrack::SetIsInactive(const bool bInIsInactive)
+void FNavigationToolSubTrack::SetIsDeactivated(const bool bInIsDeactivated)
 {
-	const bool bNewActiveState = !bInIsInactive;
+	const bool bNewActiveState = !bInIsDeactivated;
 
 	if (UMovieSceneSubSection* const ThisSubSection = Cast<UMovieSceneSubSection>(WeakSection.Get()))
 	{
@@ -111,22 +124,24 @@ void FNavigationToolSubTrack::SetIsInactive(const bool bInIsInactive)
 		}
 	}
 
-	for (ISequenceInactivableExtension* const InactivableItem : GetChildrenOfType<ISequenceInactivableExtension>())
+	for (const TViewModelPtr<IDeactivatableExtension>& InactivableItem : GetChildrenOfType<IDeactivatableExtension>())
 	{
-		InactivableItem->SetIsInactive(bInIsInactive);
+		InactivableItem->SetIsDeactivated(bInIsDeactivated);
 	}
 }
 
 EItemMarkerVisibility FNavigationToolSubTrack::GetMarkerVisibility() const
 {
-	const ENavigationToolCompareState State = CompareChildrenItemState<IMarkerVisibilityExtension>(this,
-		[](const IMarkerVisibilityExtension* const InInactivableItem)
+	using namespace ItemUtils;
+
+	const ENavigationToolCompareState State = CompareChildrenItemState<IMarkerVisibilityExtension>(AsItemViewModelConst(),
+		[](const TViewModelPtr<IMarkerVisibilityExtension>& InItem)
 			{
-				return InInactivableItem->GetMarkerVisibility() == EItemMarkerVisibility::Visible;
+				return InItem->GetMarkerVisibility() == EItemMarkerVisibility::Visible;
 			},
-		[](const IMarkerVisibilityExtension* const InInactivableItem)
+		[](const TViewModelPtr<IMarkerVisibilityExtension>& InItem)
 			{
-				return InInactivableItem->GetMarkerVisibility() == EItemMarkerVisibility::None;
+				return InItem->GetMarkerVisibility() == EItemMarkerVisibility::None;
 			});
 
 	return static_cast<EItemMarkerVisibility>(State);
@@ -151,7 +166,7 @@ void FNavigationToolSubTrack::SetMarkerVisibility(const bool bInVisible)
 		}
 	}
 
-	for (IMarkerVisibilityExtension* const MarkerVisibilityItem : GetChildrenOfType<IMarkerVisibilityExtension>())
+	for (const TViewModelPtr<IMarkerVisibilityExtension>& MarkerVisibilityItem : GetChildrenOfType<IMarkerVisibilityExtension>())
 	{
 		MarkerVisibilityItem->SetMarkerVisibility(bInVisible);
 	}
diff --git a/Engine/Plugins/Experimental/SequenceNavigator/Source/SequenceNavigator/Private/Items/NavigationToolTrack.cpp b/Engine/Plugins/Experimental/SequenceNavigator/Source/SequenceNavigator/Private/Items/NavigationToolTrack.cpp
--- a/Engine/Plugins/Experimental/SequenceNavigator/Source/SequenceNavigator/Private/Items/NavigationToolTrack.cpp
+++ b/Engine/Plugins/Experimental/SequenceNavigator/Source/SequenceNavigator/Private/Items/NavigationToolTrack.cpp
@@ -22,13 +22,15 @@
 namespace UE::SequenceNavigator
 {
 
+UE_SEQUENCER_DEFINE_CASTABLE(FNavigationToolTrack)
+
 FNavigationToolTrack::FNavigationToolTrack(INavigationTool& InTool
-	, const FNavigationToolItemPtr& InParentItem
+	, const FNavigationToolViewModelPtr& InParentItem
 	, UMovieSceneTrack* const InTrack
 	, const TWeakObjectPtr<UMovieSceneSequence>& InSequence
 	, const TWeakObjectPtr<UMovieSceneSection>& InSection
 	, const int32 InSubSectionIndex)
-	: Super(InTool, InParentItem)
+	: FNavigationToolItem(InTool, InParentItem)
 	, WeakSequence(InSequence)
 	, WeakSection(InSection)
 	, SectionIndex(InSubSectionIndex)
@@ -61,9 +63,9 @@ FNavigationToolItemId FNavigationToolTrack::CalculateItemId() const
 		, FNavigationToolItemId::GetObjectPath(WeakTrack.Get()));
 }
 
-void FNavigationToolTrack::FindChildren(TArray<FNavigationToolItemPtr>& OutChildren, const bool bInRecursive)
+void FNavigationToolTrack::FindChildren(TArray<FNavigationToolViewModelWeakPtr>& OutWeakChildren, const bool bInRecursive)
 {
-	Super::FindChildren(OutChildren, bInRecursive);
+	FNavigationToolItem::FindChildren(OutWeakChildren, bInRecursive);
 }
 
 FText FNavigationToolTrack::GetDisplayName() const
@@ -72,7 +74,7 @@ FText FNavigationToolTrack::GetDisplayName() const
 	{
 		return Track->GetDisplayName();
 	}
-	return Super::GetDisplayName();
+	return FNavigationToolItem::GetDisplayName();
 }
 
 FText FNavigationToolTrack::GetClassName() const
@@ -136,7 +138,7 @@ void FNavigationToolTrack::Select(FNavigationToolScopedSelection& InSelection) c
 
 void FNavigationToolTrack::OnSelect()
 {
-	Super::OnSelect();
+	FNavigationToolItem::OnSelect();
 }
 
 void FNavigationToolTrack::OnDoubleClick()
@@ -165,33 +167,38 @@ void FNavigationToolTrack::OnObjectsReplaced(const TMap<UObject*, UObject*>& InR
 	}
 
 	// This handles calling OnObjectsReplaced for every child item
-	Super::OnObjectsReplaced(InReplacementMap, bInRecursive);
+	FNavigationToolItem::OnObjectsReplaced(InReplacementMap, bInRecursive);
 
 	OnTrackObjectChanged();
 }
 
-EItemSequenceLockState FNavigationToolTrack::GetLockState() const
+Sequencer::ELockableLockState FNavigationToolTrack::GetLockState() const
 {
+	using namespace Sequencer;
+	using namespace ItemUtils;
+
 	if (WeakSection.IsValid())
 	{
-		return WeakSection->IsLocked() ? EItemSequenceLockState::Locked : EItemSequenceLockState::None;
+		return WeakSection->IsLocked() ? ELockableLockState::Locked : ELockableLockState::None;
 	}
 
-	const ENavigationToolCompareState State = CompareChildrenItemState<ISequenceLockableExtension>(this,
-		[](const ISequenceLockableExtension* const InLockableItem)
+	const ENavigationToolCompareState State = CompareChildrenItemState<ILockableExtension>(AsItemViewModelConst(),
+		[](const TViewModelPtr<ILockableExtension>& InItem)
 			{
-				return InLockableItem->GetLockState() == EItemSequenceLockState::Locked;
+				return InItem->GetLockState() == ELockableLockState::Locked;
 			},
-		[](const ISequenceLockableExtension* const InLockableItem)
+		[](const TViewModelPtr<ILockableExtension>& InItem)
 			{
-				return InLockableItem->GetLockState() != EItemSequenceLockState::Locked;
+				return InItem->GetLockState() != ELockableLockState::Locked;
 			});
 
-	return static_cast<EItemSequenceLockState>(State);
+	return static_cast<ELockableLockState>(State);
 }
 
 void FNavigationToolTrack::SetIsLocked(const bool bInIsLocked)
 {
+	using namespace Sequencer;
+
 	if (WeakSection.IsValid())
 	{
 		if (WeakSection->IsLocked() != bInIsLocked)
@@ -201,7 +208,7 @@ void FNavigationToolTrack::SetIsLocked(const bool bInIsLocked)
 		}
 	}
 
-	for (ISequenceLockableExtension* const LockableItem : GetChildrenOfType<ISequenceLockableExtension>())
+	for (const TViewModelPtr<ILockableExtension>& LockableItem : GetChildrenOfType<ILockableExtension>())
 	{
 		LockableItem->SetIsLocked(bInIsLocked);
 	}
@@ -288,7 +295,7 @@ int32 FNavigationToolTrack::GetSectionIndex() const
 
 UE::Sequencer::TViewModelPtr<UE::Sequencer::FTrackModel> FNavigationToolTrack::GetViewModel() const
 {
-	using namespace UE::Sequencer;
+	using namespace Sequencer;
 
 	UMovieSceneTrack* const Track = GetTrack();
 	if (!Track)
diff --git a/Engine/Plugins/Experimental/SequenceNavigator/Source/SequenceNavigator/Private/Items/NavigationToolTreeRoot.cpp b/Engine/Plugins/Experimental/SequenceNavigator/Source/SequenceNavigator/Private/Items/NavigationToolTreeRoot.cpp
--- a/Engine/Plugins/Experimental/SequenceNavigator/Source/SequenceNavigator/Private/Items/NavigationToolTreeRoot.cpp
+++ b/Engine/Plugins/Experimental/SequenceNavigator/Source/SequenceNavigator/Private/Items/NavigationToolTreeRoot.cpp
@@ -9,33 +9,33 @@
 namespace UE::SequenceNavigator
 {
 
+UE_SEQUENCER_DEFINE_CASTABLE(FNavigationToolTreeRoot)
+	
 FNavigationToolTreeRoot::FNavigationToolTreeRoot(INavigationTool& InTool)
 	: FNavigationToolItem(InTool, nullptr)
 {
 }
 
-void FNavigationToolTreeRoot::FindChildren(TArray<FNavigationToolItemPtr>& OutChildren, const bool bInRecursive)
+void FNavigationToolTreeRoot::FindChildren(TArray<FNavigationToolViewModelWeakPtr>& OutWeakChildren, const bool bInRecursive)
 {
-	Super::FindChildren(OutChildren, bInRecursive);
-
-	const TSharedRef<FNavigationToolTreeRoot> SharedThisRef = SharedThis(this);
+	FNavigationToolItem::FindChildren(OutWeakChildren, bInRecursive);
 
 	// Get items to add from providers
-	GetOwnerTool().ForEachProvider([this, &OutChildren, bInRecursive, &SharedThisRef]
+	GetOwnerTool().ForEachProvider([this, &OutWeakChildren, bInRecursive]
 		(const TSharedRef<FNavigationToolProvider>& InToolProvider)
 		{
-			TArray<FNavigationToolItemPtr> ThisProviderItems;
-			InToolProvider->OnExtendItemChildren(Tool, SharedThisRef, ThisProviderItems, bInRecursive);
+			TArray<FNavigationToolViewModelWeakPtr> WeakProviderItems;
+			InToolProvider->OnExtendItemChildren(Tool, SharedThis(this), WeakProviderItems, bInRecursive);
 
-			OutChildren.Append(ThisProviderItems);
+			OutWeakChildren.Append(WeakProviderItems);
 
 			return true;
 		});
 }
 
-bool FNavigationToolTreeRoot::CanAddChild(const FNavigationToolItemPtr& InChild) const
+bool FNavigationToolTreeRoot::CanAddChild(const FNavigationToolViewModelPtr& InChild) const
 {
-	return Super::CanAddChild(InChild) && InChild->CanBeTopLevel();
+	return FNavigationToolItem::CanAddChild(InChild) && InChild->CanBeTopLevel();
 }
 
 bool FNavigationToolTreeRoot::IsAllowedInTool() const
@@ -76,14 +76,14 @@ TSharedRef<SWidget> FNavigationToolTreeRoot::GenerateLabelWidget(const TSharedRe
 
 TOptional<EItemDropZone> FNavigationToolTreeRoot::CanAcceptDrop(const FDragDropEvent& InDragDropEvent, const EItemDropZone InDropZone)
 {
-	return Super::CanAcceptDrop(InDragDropEvent, InDropZone);
+	return FNavigationToolItem::CanAcceptDrop(InDragDropEvent, InDropZone);
 }
 
 FReply FNavigationToolTreeRoot::AcceptDrop(const FDragDropEvent& InDragDropEvent, const EItemDropZone InDropZone)
 {
 	// Force tree root to always be OntoItem so we always create new items as children of root
 	constexpr EItemDropZone ForcedDropZone = EItemDropZone::OntoItem;
-	return Super::AcceptDrop(InDragDropEvent, ForcedDropZone);
+	return FNavigationToolItem::AcceptDrop(InDragDropEvent, ForcedDropZone);
 }
 
 FNavigationToolItemId FNavigationToolTreeRoot::CalculateItemId() const
diff --git a/Engine/Plugins/Experimental/SequenceNavigator/Source/SequenceNavigator/Private/Items/NavigationToolTreeRoot.h b/Engine/Plugins/Experimental/SequenceNavigator/Source/SequenceNavigator/Private/Items/NavigationToolTreeRoot.h
--- a/Engine/Plugins/Experimental/SequenceNavigator/Source/SequenceNavigator/Private/Items/NavigationToolTreeRoot.h
+++ b/Engine/Plugins/Experimental/SequenceNavigator/Source/SequenceNavigator/Private/Items/NavigationToolTreeRoot.h
@@ -3,6 +3,9 @@
 #pragma once
 
 #include "Items/NavigationToolItem.h"
+#include "MVVM/ICastable.h"
+
+#define UE_API SEQUENCENAVIGATOR_API
 
 namespace UE::SequenceNavigator
 {
@@ -13,17 +16,18 @@ class INavigationTool;
  * The Item that ensures that every item (except self) has a parent in the hierarchy to make it easier
  * to handle. This is not really a visual item, so it can't appear in the Navigation Tool view.
  */
-class FNavigationToolTreeRoot final : public FNavigationToolItem
+class FNavigationToolTreeRoot final
+	: public FNavigationToolItem
 {
 public:
-	UE_NAVIGATIONTOOL_INHERITS_WITH_SUPER(FNavigationToolTreeRoot, FNavigationToolItem)
+	UE_SEQUENCER_DECLARE_CASTABLE_API(UE_API, FNavigationToolTreeRoot
+		, FNavigationToolItem)
 
 	FNavigationToolTreeRoot(INavigationTool& InTool);
-	virtual ~FNavigationToolTreeRoot() override = default;
 
 	//~ Begin INavigationToolItem
-	virtual void FindChildren(TArray<FNavigationToolItemPtr>& OutChildren, const bool bInRecursive) override;
-	virtual bool CanAddChild(const FNavigationToolItemPtr& InChild) const override;
+	virtual void FindChildren(TArray<FNavigationToolViewModelWeakPtr>& OutWeakChildren, const bool bInRecursive) override;
+	virtual bool CanAddChild(const FNavigationToolViewModelPtr& InChild) const override;
 	virtual bool IsAllowedInTool() const override;
 	virtual FText GetDisplayName() const override;
 	virtual FText GetClassName() const override;
@@ -41,3 +45,5 @@ protected:
 };
 
 } // namespace UE::SequenceNavigator
+
+#undef UE_API
diff --git a/Engine/Plugins/Experimental/SequenceNavigator/Source/SequenceNavigator/Private/Menus/NavigationToolItemContextMenu.cpp b/Engine/Plugins/Experimental/SequenceNavigator/Source/SequenceNavigator/Private/Menus/NavigationToolItemContextMenu.cpp
--- a/Engine/Plugins/Experimental/SequenceNavigator/Source/SequenceNavigator/Private/Menus/NavigationToolItemContextMenu.cpp
+++ b/Engine/Plugins/Experimental/SequenceNavigator/Source/SequenceNavigator/Private/Menus/NavigationToolItemContextMenu.cpp
@@ -14,13 +14,16 @@
 namespace UE::SequenceNavigator
 {
 
+using namespace Sequencer;
+
 FName FNavigationToolItemContextMenu::GetMenuName()
 {
 	static const FName MenuName(TEXT("SequenceNavigator.ItemContextMenu"));
 	return MenuName;
 }
 
-TSharedRef<SWidget> FNavigationToolItemContextMenu::CreateMenu(const TSharedRef<FNavigationToolView>& InToolView, const TArray<FNavigationToolItemPtr>& InItemList)
+TSharedRef<SWidget> FNavigationToolItemContextMenu::CreateMenu(const TSharedRef<FNavigationToolView>& InToolView
+	, const TArray<FNavigationToolViewModelWeakPtr>& InWeakItemList)
 {
 	UToolMenus* const ToolMenus = UToolMenus::Get();
 
@@ -33,7 +36,7 @@ TSharedRef<SWidget> FNavigationToolItemContextMenu::CreateMenu(const TSharedRef<
 	}
 
 	UNavigationToolItemMenuContext* const ContextObject = NewObject<UNavigationToolItemMenuContext>();
-	ContextObject->Init(InToolView->GetOwnerTool(), InItemList);
+	ContextObject->Init(InToolView->GetOwnerTool(), InWeakItemList);
 
 	FToolMenuContext Context(InToolView->GetBaseCommandList(), nullptr, ContextObject);
 	return ToolMenus->GenerateWidget(MenuName, Context);
diff --git a/Engine/Plugins/Experimental/SequenceNavigator/Source/SequenceNavigator/Private/Menus/NavigationToolItemContextMenu.h b/Engine/Plugins/Experimental/SequenceNavigator/Source/SequenceNavigator/Private/Menus/NavigationToolItemContextMenu.h
--- a/Engine/Plugins/Experimental/SequenceNavigator/Source/SequenceNavigator/Private/Menus/NavigationToolItemContextMenu.h
+++ b/Engine/Plugins/Experimental/SequenceNavigator/Source/SequenceNavigator/Private/Menus/NavigationToolItemContextMenu.h
@@ -13,13 +13,15 @@ namespace UE::SequenceNavigator
 {
 
 class FNavigationToolView;
+class INavigationToolItem;
 
 class FNavigationToolItemContextMenu : public TSharedFromThis<FNavigationToolItemContextMenu>
 {
 public:
 	static FName GetMenuName();
 
-	TSharedRef<SWidget> CreateMenu(const TSharedRef<FNavigationToolView>& InToolView, const TArray<FNavigationToolItemPtr>& InItemList);
+	TSharedRef<SWidget> CreateMenu(const TSharedRef<FNavigationToolView>& InToolView
+		, const TArray<FNavigationToolViewModelWeakPtr>& InWeakItemList);
 
 protected:
 	static void PopulateMenu(UToolMenu* const InMenu);
diff --git a/Engine/Plugins/Experimental/SequenceNavigator/Source/SequenceNavigator/Private/Menus/NavigationToolToolbarMenu.cpp b/Engine/Plugins/Experimental/SequenceNavigator/Source/SequenceNavigator/Private/Menus/NavigationToolToolbarMenu.cpp
--- a/Engine/Plugins/Experimental/SequenceNavigator/Source/SequenceNavigator/Private/Menus/NavigationToolToolbarMenu.cpp
+++ b/Engine/Plugins/Experimental/SequenceNavigator/Source/SequenceNavigator/Private/Menus/NavigationToolToolbarMenu.cpp
@@ -188,7 +188,7 @@ void FNavigationToolToolbarMenu::CreateItemViewOptionsMenu(UToolMenu* const InMe
 			continue;
 		}
 
-		const FName ItemProxyTypeName = TemplateProxyItem->GetTypeId().ToName();
+		const FName ItemProxyTypeName = TemplateProxyItem->GetTypeTable().GetTypeName();
 
 		// Type mismatch here can continue but recommended to be addressed for Proxy Items / Factories
 		// not correctly overriding Type Name Functions
diff --git a/Engine/Plugins/Experimental/SequenceNavigator/Source/SequenceNavigator/Private/NavigationTool.cpp b/Engine/Plugins/Experimental/SequenceNavigator/Source/SequenceNavigator/Private/NavigationTool.cpp
--- a/Engine/Plugins/Experimental/SequenceNavigator/Source/SequenceNavigator/Private/NavigationTool.cpp
+++ b/Engine/Plugins/Experimental/SequenceNavigator/Source/SequenceNavigator/Private/NavigationTool.cpp
@@ -10,7 +10,7 @@
 #include "Framework/MultiBox/MultiBoxBuilder.h"
 #include "GameFramework/Actor.h"
 #include "ISequencer.h"
-#include "ItemActions/NavigationToolItemAction.h"
+#include "ItemActions/INavigationToolItemAction.h"
 #include "ItemProxies/NavigationToolItemProxyRegistry.h"
 #include "Items/NavigationToolBinding.h"
 #include "Items/NavigationToolItemUtils.h"
@@ -34,11 +34,11 @@
 
 #define LOCTEXT_NAMESPACE "NavigationTool"
 
-using namespace UE::Sequencer;
-
 namespace UE::SequenceNavigator
 {
 
+using namespace Sequencer;
+
 FNavigationTool::FNavigationTool(const TWeakPtr<ISequencer>& InWeakSequencer)
 	: WeakSequencer(InWeakSequencer)
 	, ToolTab(MakeShared<FNavigationToolTab>(*this))
@@ -187,10 +187,10 @@ TArray<FName> FNavigationTool::GetRegisteredItemProxyTypeNames() const
 	return OutItemProxyTypeNames;
 }
 
-void FNavigationTool::GetItemProxiesForItem(const FNavigationToolItemPtr& InItem, TArray<TSharedPtr<FNavigationToolItemProxy>>& OutItemProxies)
+void FNavigationTool::GetItemProxiesForItem(const FNavigationToolViewModelPtr& InItem, TArray<TSharedPtr<FNavigationToolItemProxy>>& OutItemProxies)
 {
 	// No Item Proxy support for Root
-	if (!InItem.IsValid() || InItem == RootItem)
+	if (!InItem.IsValid() || InItem.ToSharedRef() == RootItem.ToSharedRef())
 	{
 		return;
 	}
@@ -204,10 +204,13 @@ void FNavigationTool::GetItemProxiesForItem(const FNavigationToolItemPtr& InItem
 			return true;
 		});
 
-	// Clean up any invalid Item Proxy
-	OutItemProxies.RemoveAll([](const TSharedPtr<FNavigationToolItemProxy>& InItemProxy) { return !InItemProxy.IsValid(); });
+	// Clean up any invalid item proxy
+	OutItemProxies.RemoveAll([](const TSharedPtr<FNavigationToolItemProxy>& InItemProxy)
+		{
+			return !InItemProxy.IsValid();
+		});
 
-	// Sort Proxies by their Priority
+	// Sort proxies by their priority
 	OutItemProxies.Sort([](const TSharedPtr<FNavigationToolItemProxy>& InItemProxyA, const TSharedPtr<FNavigationToolItemProxy>& InItemProxyB)
 		{
 			return InItemProxyA->GetPriority() > InItemProxyB->GetPriority();
@@ -216,13 +219,13 @@ void FNavigationTool::GetItemProxiesForItem(const FNavigationToolItemPtr& InItem
 
 INavigationToolItemProxyFactory* FNavigationTool::GetItemProxyFactory(const FName InItemProxyTypeName) const
 {
-	// First look for the Registry in Navigation Tool
+	// First look for the registry in Navigation Tool
 	if (INavigationToolItemProxyFactory* const Factory = ItemProxyRegistry.GetItemProxyFactory(InItemProxyTypeName))
 	{
 		return Factory;
 	}
 
-	// Fallback to finding the Factory in the Module if the Navigation Tool did not find it
+	// Fallback to finding the factory in the module if the Navigation Tool did not find it
 	return FNavigationToolExtender::GetItemProxyRegistry().GetItemProxyFactory(InItemProxyTypeName);
 }
 
@@ -356,7 +359,7 @@ bool FNavigationTool::IsObjectAllowedInTool(const UObject* const InObject) const
 	return false;
 }
 
-void FNavigationTool::RegisterItem(const FNavigationToolItemPtr& InItem)
+void FNavigationTool::RegisterItem(const FNavigationToolViewModelPtr& InItem)
 {
 	if (!InItem.IsValid())
 	{
@@ -364,7 +367,7 @@ void FNavigationTool::RegisterItem(const FNavigationToolItemPtr& InItem)
 	}
 
 	const FNavigationToolItemId ItemId = InItem->GetItemId();
-	const FNavigationToolItemPtr ExistingItem = FindItem(ItemId);
+	const FNavigationToolViewModelPtr ExistingItem = FindItem(ItemId);
 
 	// if there's no existing item or the existing item does not match its new replacement,
 	// then call OnItemRegistered and Refresh Navigation Tool
@@ -380,7 +383,7 @@ void FNavigationTool::RegisterItem(const FNavigationToolItemPtr& InItem)
 
 void FNavigationTool::UnregisterItem(const FNavigationToolItemId& InItemId)
 {
-	const FNavigationToolItemPtr FoundItem = FindItem(InItemId);
+	const FNavigationToolViewModelPtr FoundItem = FindItem(InItemId);
 	if (!FoundItem.IsValid())
 	{
 		return;
@@ -403,27 +406,27 @@ void FNavigationTool::Refresh()
 	TGuardValue RefreshGuard(bRefreshing, true);
 	bRefreshRequested = false;
 
-	// Flush Pending Actions
+	// Flush pending actions
 	{
-		// Make a Transaction if there's a Pending Action Requesting it
-		TUniquePtr<FScopedTransaction> Transaction;
+		// Make a transaction if there's a pending action requesting it
+		TUniquePtr<FScopedTransaction> Transaction = nullptr;
 
 		if (!GIsTransacting)
 		{
 			const bool bShouldTransact = PendingActions.ContainsByPredicate(
-				[](const TSharedPtr<INavigationToolAction>& InAction)
+				[](const TSharedPtr<INavigationToolItemAction>& InAction)
 				{
 					return InAction.IsValid() && InAction->ShouldTransact();
 				});
 
 			if (bShouldTransact)
 			{
-				Transaction = MakeUnique<FScopedTransaction>(LOCTEXT("ItemAction", "Sequence Navigator Item Action"));
+				Transaction = MakeUnique<FScopedTransaction>(LOCTEXT("ItemAction", "Item Action"));
 			}
 		}
 
-		// Execute Pending Actions
-		for (const TSharedPtr<INavigationToolAction>& Action : PendingActions)
+		// Execute pending actions
+		for (const TSharedPtr<INavigationToolItemAction>& Action : PendingActions)
 		{
 			if (Action.IsValid())
 			{
@@ -434,17 +437,17 @@ void FNavigationTool::Refresh()
 		PendingActions.Empty();
 	}
 
-	// Save the Current Item ordering before Refreshing Children
-	// Do not reset tree just yet as there might be actors that still need to be discovered on the next pass
-	// This is done to save the items added from the Queued Actions and be considered when adding new items from Discovery
+	// Save the current item ordering before refreshing children.
+	// Do not reset the tree just yet as there might be actors that still need to be discovered on the next pass.
+	// This is done to save the items added from the queued actions and be considered when adding new items from discovery.
 	ForEachProvider([this](const TSharedRef<FNavigationToolProvider>& InProvider)
 		{
 			InProvider->SaveSerializedTree(*this, /*bInResetTree*/false);
 			return true;
 		});
 
-	// Refresh each item's children. This also Updates each Child's parent var
-	ForEachItem([](const FNavigationToolItemPtr& InItem)
+	// Refresh each item's children. This also updates each child's parent var.
+	ForEachItem([](const FNavigationToolViewModelPtr& InItem)
 		{
 			InItem->RefreshChildren();
 		});
@@ -456,7 +459,7 @@ void FNavigationTool::Refresh()
 			InToolView->Refresh();
 		});
 
-	// Save so that the Tree is updated to the latest Navigation Tool State
+	// Save so that the tree is updated to the latest Navigation Tool state
 	ForEachProvider([this](const TSharedRef<FNavigationToolProvider>& InProvider)
 		{
 			InProvider->SaveSerializedTree(*this, /*bInResetTree*/false);
@@ -464,22 +467,22 @@ void FNavigationTool::Refresh()
 		});
 }
 
-TSharedRef<FNavigationToolItem> FNavigationTool::GetTreeRoot() const
+FNavigationToolViewModelWeakPtr FNavigationTool::GetTreeRoot() const
 {
 	return RootItem;
 }
 
-FNavigationToolItemPtr FNavigationTool::FindItem(const FNavigationToolItemId& InItemId) const
+FNavigationToolViewModelPtr FNavigationTool::FindItem(const FNavigationToolItemId& InItemId) const
 {
 	if (bIteratingItemMap)
 	{
-		if (const FNavigationToolItemPtr* const FoundItem = ItemsPendingAdd.Find(InItemId))
+		if (const FNavigationToolViewModelPtr* const FoundItem = ItemsPendingAdd.Find(InItemId))
 		{
 			return *FoundItem;
 		}
 	}
 
-	if (const FNavigationToolItemPtr* const FoundItem = ItemMap.Find(InItemId))
+	if (const FNavigationToolViewModelPtr* const FoundItem = ItemMap.Find(InItemId))
 	{
 		return *FoundItem;
 	}
@@ -487,24 +490,24 @@ FNavigationToolItemPtr FNavigationTool::FindItem(const FNavigationToolItemId& In
 	return nullptr;
 }
 
-TArray<FNavigationToolItemPtr> FNavigationTool::TryFindItems(const UE::Sequencer::FViewModelPtr& InViewModel) const
+TArray<FNavigationToolViewModelPtr> FNavigationTool::TryFindItems(const FViewModelPtr& InSequencerOutlinerViewModel) const
 {
-	if (const TViewModelPtr<FTrackModel> TrackModel = InViewModel.ImplicitCast())
+	if (const TViewModelPtr<FTrackModel> TrackModel = InSequencerOutlinerViewModel.ImplicitCast())
 	{
 		return FindItemsFromMovieSceneObject(TrackModel->GetTrack());
 	}
 
-	if (const TViewModelPtr<FTrackRowModel> TrackRowModel = InViewModel.ImplicitCast())
+	if (const TViewModelPtr<FTrackRowModel> TrackRowModel = InSequencerOutlinerViewModel.ImplicitCast())
 	{
 		return FindItemsFromMovieSceneObject(TrackRowModel->GetTrack());
 	}
 
-	if (const TViewModelPtr<FObjectBindingModel> ObjectBindingModel = InViewModel.ImplicitCast())
+	if (const TViewModelPtr<FObjectBindingModel> ObjectBindingModel = InSequencerOutlinerViewModel.ImplicitCast())
 	{
 		return FindItemsFromObjectGuid(ObjectBindingModel->GetObjectGuid());
 	}
 
-	if (const TViewModelPtr<FSectionModel> SectionModel = InViewModel.ImplicitCast())
+	if (const TViewModelPtr<FSectionModel> SectionModel = InSequencerOutlinerViewModel.ImplicitCast())
 	{
 		if (UMovieSceneSubSection* const SubSection = Cast<UMovieSceneSubSection>(SectionModel->GetSection()))
 		{
@@ -567,18 +570,18 @@ void FNavigationTool::Tick(const float InDeltaTime)
 	// Select Items Pending Selection
 	if (ItemsLastSelected.IsValid())
 	{
-		TArray<FNavigationToolItemPtr> ItemsToSelect;
+		TArray<FNavigationToolViewModelWeakPtr> WeakItemsToSelect;
 
-		for (const TWeakPtr<INavigationToolItem>& WeakItem : *ItemsLastSelected)
+		for (const FNavigationToolViewModelWeakPtr& WeakItem : *ItemsLastSelected)
 		{
-			if (const TSharedPtr<INavigationToolItem> Item = WeakItem.Pin())
+			if (const FNavigationToolViewModelPtr Item = WeakItem.Pin())
 			{
-				ItemsToSelect.Add(Item);
+				WeakItemsToSelect.Add(Item);
 			}
 		}
 
 		// Only Scroll Into View, don't signal selection since we just came from the selection notify itself
-		SelectItems(ItemsToSelect, ENavigationToolItemSelectionFlags::ScrollIntoView);
+		SelectItems(WeakItemsToSelect, ENavigationToolItemSelectionFlags::ScrollIntoView);
 
 		ItemsLastSelected.Reset();
 	}
@@ -589,31 +592,33 @@ void FNavigationTool::Tick(const float InDeltaTime)
 		});
 }
 
-void FNavigationTool::DeleteItems(TArray<FNavigationToolItemPtr> InItems)
+void FNavigationTool::DeleteItems(const TArray<FNavigationToolViewModelWeakPtr>& InWeakItems)
 {
-	SortItems(InItems);
+	//SortItems(InWeakItems);
 
-	for (TArray<FNavigationToolItemPtr>::TIterator It(InItems); It; ++It)
-	{
-		FNavigationToolItemPtr Item = *It;
+	TArray<FNavigationToolViewModelPtr> ItemsToRemove;
 
-		if (!Item.IsValid() || !Item->CanDelete())
+	for (const FNavigationToolViewModelWeakPtr& WeakItem : InWeakItems)
+	{
+		const FNavigationToolViewModelPtr Item = WeakItem.Pin();
+		if (Item.IsValid() || Item->CanDelete())
 		{
-			It.RemoveCurrent();
+			ItemsToRemove.Add(Item);
 		}
 	}
 
-	if (InItems.IsEmpty())
+	if (ItemsToRemove.IsEmpty())
 	{
 		return;
 	}
 
-	FScopedTransaction DeleteTransaction(LOCTEXT("ItemDeleteAction", "Sequence Navigator Delete Item(s)"), !GIsTransacting);
+	const FScopedTransaction DeleteTransaction(LOCTEXT("ItemDeleteAction", "Delete Item(s)"), !GIsTransacting);
 
 	bool bRequestRefresh = false;
-	for (FNavigationToolItemPtr& Item : InItems)
+
+	for (const FNavigationToolViewModelPtr& Item : ItemsToRemove)
 	{
-		if (Item && Item->Delete())
+		if (Item.IsValid() && Item->Delete())
 		{
 			bRequestRefresh = true;
 		}
@@ -665,7 +670,7 @@ void FNavigationTool::ForEachToolView(const TFunction<void(const TSharedRef<FNav
 	}
 }
 
-void FNavigationTool::EnqueueItemActions(TArray<TSharedPtr<INavigationToolAction>>&& InItemActions) noexcept
+void FNavigationTool::EnqueueItemActions(TArray<TSharedPtr<INavigationToolItemAction>>&& InItemActions) noexcept
 {
 	PendingActions.Append(MoveTemp(InItemActions));
 }
@@ -690,7 +695,7 @@ bool FNavigationTool::NeedsRefresh() const
 	return false;
 }
 
-TOptional<FColor> FNavigationTool::FindItemColor(const FNavigationToolItemPtr& InItem, bool bRecurseParent) const
+TOptional<FColor> FNavigationTool::FindItemColor(const FNavigationToolViewModelPtr& InItem, bool bRecurseParent) const
 {
 	if (!InItem.IsValid())
 	{
@@ -723,7 +728,7 @@ TOptional<FColor> FNavigationTool::FindItemColor(const FNavigationToolItemPtr& I
 	return TOptional<FColor>();
 }
 
-void FNavigationTool::SetItemColor(const FNavigationToolItemPtr& InItem, const FColor& InColor)
+void FNavigationTool::SetItemColor(const FNavigationToolViewModelPtr& InItem, const FColor& InColor)
 {
 	if (!InItem.IsValid())
 	{
@@ -768,12 +773,11 @@ void FNavigationTool::SetItemColor(const FNavigationToolItemPtr& InItem, const F
 			SaveState->ItemColorMap.Add(InItem->GetItemId().GetStringId(), InColor);
 		}
 
-		//Recurse Children to Remove their Mapping if same color with parent
-		TArray<FNavigationToolItemPtr> RemainingChildren = InItem->GetChildren();
-		while (RemainingChildren.Num() > 0)
+		// Recurse children to remove their mapping if same color with parent
+		TArray<FNavigationToolViewModelWeakPtr> WeakRemainingChildren = InItem->GetChildren();
+		while (WeakRemainingChildren.Num() > 0)
 		{
-			const FNavigationToolItemPtr Child = RemainingChildren.Pop();
-			if (Child.IsValid())
+			if (const FNavigationToolViewModelPtr Child = WeakRemainingChildren.Pop().Pin())
 			{
 				const TOptional<FColor> ChildColor = FindItemColor(Child, false);
 				if (ChildColor.IsSet() && ChildColor == InColor)
@@ -782,7 +786,7 @@ void FNavigationTool::SetItemColor(const FNavigationToolItemPtr& InItem, const F
 
 					//Only check Grand Children if Child had same Color.
 					//Allow the situation where Parent is ColorA, Child is ColorB and GrandChild is ColorA.
-					RemainingChildren.Append(Child->GetChildren());
+					WeakRemainingChildren.Append(Child->GetChildren());
 				}
 			}
 		}
@@ -790,7 +794,7 @@ void FNavigationTool::SetItemColor(const FNavigationToolItemPtr& InItem, const F
 	}
 }
 
-void FNavigationTool::RemoveItemColor(const FNavigationToolItemPtr& InItem)
+void FNavigationTool::RemoveItemColor(const FNavigationToolViewModelPtr& InItem)
 {
 	if (!InItem.IsValid())
 	{
@@ -816,7 +820,7 @@ void FNavigationTool::RemoveItemColor(const FNavigationToolItemPtr& InItem)
 	}
 }
 
-void FNavigationTool::NotifyItemIdChanged(const FNavigationToolItemId& OldId, const FNavigationToolItemPtr& InItem)
+void FNavigationTool::NotifyItemIdChanged(const FNavigationToolItemId& OldId, const FNavigationToolViewModelPtr& InItem)
 {
 	const FNavigationToolItemId NewId = InItem->GetItemId();
 	if (OldId == NewId)
@@ -824,7 +828,7 @@ void FNavigationTool::NotifyItemIdChanged(const FNavigationToolItemId& OldId, co
 		return;
 	}
 
-	const FNavigationToolItemPtr FoundItem = FindItem(OldId);
+	const FNavigationToolViewModelPtr FoundItem = FindItem(OldId);
 	if (FoundItem.IsValid() && FoundItem == InItem)
 	{
 		AddItem(InItem);
@@ -833,7 +837,7 @@ void FNavigationTool::NotifyItemIdChanged(const FNavigationToolItemId& OldId, co
 	}
 }
 
-void FNavigationTool::NotifyToolItemRenamed(const FNavigationToolItemPtr& InItem)
+void FNavigationTool::NotifyToolItemRenamed(const FNavigationToolViewModelPtr& InItem)
 {
 	ForEachProvider([&InItem](const TSharedRef<FNavigationToolProvider>& InProvider)
 		{
@@ -842,7 +846,7 @@ void FNavigationTool::NotifyToolItemRenamed(const FNavigationToolItemPtr& InItem
 		});
 }
 
-void FNavigationTool::NotifyToolItemDeleted(const FNavigationToolItemPtr& InItem)
+void FNavigationTool::NotifyToolItemDeleted(const FNavigationToolViewModelPtr& InItem)
 {
 	ForEachProvider([&InItem](const TSharedRef<INavigationToolProvider>& InProvider)
 		{
@@ -851,29 +855,30 @@ void FNavigationTool::NotifyToolItemDeleted(const FNavigationToolItemPtr& InItem
 		});
 }
 
-TArray<FNavigationToolItemPtr> FNavigationTool::GetSelectedItems(const bool bInNormalizeToTopLevelSelections) const
+TArray<FNavigationToolViewModelWeakPtr> FNavigationTool::GetSelectedItems(const bool bInNormalizeToTopLevelSelections) const
 {
 	if (const TSharedPtr<INavigationToolView>& ToolView = GetMostRecentToolView())
 	{
-		TArray<FNavigationToolItemPtr> SelectedItems = ToolView->GetSelectedItems();
+		TArray<FNavigationToolViewModelWeakPtr> WeakSelectedItems = ToolView->GetSelectedItems();
 
 		if (bInNormalizeToTopLevelSelections)
 		{
-			NormalizeToTopLevelSelections(SelectedItems);
-			return SelectedItems;
+			NormalizeToTopLevelSelections(WeakSelectedItems);
+			return WeakSelectedItems;
 		}
 
-		return SelectedItems;
+		return WeakSelectedItems;
 	}
 
-	return TArray<FNavigationToolItemPtr>();
+	return {};
 }
 
-void FNavigationTool::SelectItems(const TArray<FNavigationToolItemPtr>& InItems, const ENavigationToolItemSelectionFlags InFlags) const
+void FNavigationTool::SelectItems(const TArray<FNavigationToolViewModelWeakPtr>& InWeakItems
+	, const ENavigationToolItemSelectionFlags InFlags) const
 {
-	ForEachToolView([InItems, InFlags](const TSharedPtr<INavigationToolView>& InToolView)
+	ForEachToolView([InWeakItems, InFlags](const TSharedPtr<INavigationToolView>& InToolView)
 		{
-			InToolView->SelectItems(InItems, InFlags);
+			InToolView->SelectItems(InWeakItems, InFlags);
 		});
 }
 
@@ -885,14 +890,14 @@ void FNavigationTool::ClearItemSelection(const bool bInSignalSelectionChange) co
 		});
 }
 
-FNavigationToolItemPtr FNavigationTool::FindLowestCommonAncestor(const TArray<FNavigationToolItemPtr>& Items)
+FNavigationToolViewModelPtr FNavigationTool::FindLowestCommonAncestor(const TArray<FNavigationToolViewModelPtr>& Items)
 {
-	TSet<FNavigationToolItemPtr> IntersectedAncestors;
+	TSet<FNavigationToolViewModelPtr> IntersectedAncestors;
 
-	for (const FNavigationToolItemPtr& Item : Items)
+	for (const FNavigationToolViewModelPtr& Item : Items)
 	{
-		FNavigationToolItemPtr Parent = Item->GetParent();
-		TSet<FNavigationToolItemPtr> ItemAncestors;
+		FNavigationToolViewModelPtr Parent = Item->GetParent();
+		TSet<FNavigationToolViewModelPtr> ItemAncestors;
 
 		//Add all Item's Ancestors
 		while (Parent.IsValid())
@@ -918,11 +923,21 @@ FNavigationToolItemPtr FNavigationTool::FindLowestCommonAncestor(const TArray<FN
 		}
 	}
 
-	FNavigationToolItemPtr LowestCommonAncestor;
-	for (const FNavigationToolItemPtr& Item : IntersectedAncestors)
+	FNavigationToolViewModelPtr LowestCommonAncestor;
+	for (const FNavigationToolViewModelPtr& Item : IntersectedAncestors)
 	{
-		//Find Item with most tree height (i.e. lowest down the tree, closer to the selected nodes)
-		if (!LowestCommonAncestor.IsValid() || Item->GetItemTreeHeight() > LowestCommonAncestor->GetItemTreeHeight())
+		const TViewModelPtr<IOutlinerExtension> OutlinerExtension = Item.ImplicitCast();
+		if (!OutlinerExtension.IsValid())
+		{
+			continue;
+		}
+
+		const TViewModelPtr<IOutlinerExtension> LowestOutlinerExtension = LowestCommonAncestor.ImplicitCast();
+
+		// Find item with most tree height (i.e. lowest down the tree, closer to the selected nodes)
+		if (!LowestCommonAncestor.IsValid()
+			|| (LowestOutlinerExtension.IsValid()
+				&& OutlinerExtension->GetOutlinerSizing().Height > LowestOutlinerExtension->GetOutlinerSizing().Height))
 		{
 			LowestCommonAncestor = Item;
 		}
@@ -930,28 +945,35 @@ FNavigationToolItemPtr FNavigationTool::FindLowestCommonAncestor(const TArray<FN
 	return LowestCommonAncestor;
 }
 
-void FNavigationTool::SortItems(TArray<FNavigationToolItemPtr>& OutToolItems, const bool bInReverseOrder)
+void FNavigationTool::SortItems(TArray<FNavigationToolViewModelWeakPtr>& WeakItems
+	, const bool bInReverseOrder)
 {
-	OutToolItems.Sort([bInReverseOrder](const FNavigationToolItemPtr& InItemA, const FNavigationToolItemPtr& InItemB)
+	WeakItems.Sort([bInReverseOrder]
+		(const FNavigationToolViewModelWeakPtr& InWeakItemA, const FNavigationToolViewModelWeakPtr& InWeakItemB)
 		{
-			return CompareToolItemOrder(InItemA, InItemB) != bInReverseOrder;
+			return ItemUtils::CompareToolItemOrder(InWeakItemA.Pin(), InWeakItemB.Pin()) != bInReverseOrder;
 		});
 }
 
-void FNavigationTool::NormalizeToTopLevelSelections(TArray<FNavigationToolItemPtr>& InOutItems)
+void FNavigationTool::NormalizeToTopLevelSelections(TArray<FNavigationToolViewModelWeakPtr>& InWeakItems)
 {
-	if (InOutItems.IsEmpty())
+	if (InWeakItems.IsEmpty())
 	{
 		return;
 	}
 
-	//Set for Quick Lookup
-	const TSet<FNavigationToolItemPtr> SelectedItemSet(InOutItems);
+	const TSet<FNavigationToolViewModelWeakPtr> SelectedItemSet(InWeakItems);
 
-	//Normalize Selection: Remove all Items that have Parents that are in the Selection. Swapping since we're sorting afterwards
-	InOutItems.RemoveAllSwap([&SelectedItemSet](const FNavigationToolItemPtr& Item)
+	// Normalize Selection: Remove all items that have parents that are in the selection.
+	// Swapping since we're sorting afterwards.
+	InWeakItems.RemoveAllSwap([&SelectedItemSet](const FNavigationToolViewModelWeakPtr& InWeakItem)
 	{
-		FNavigationToolItemPtr Parent = Item->GetParent();
+		const FNavigationToolViewModelPtr Item = InWeakItem.Pin();
+		if (!Item.IsValid())
+		{
+			return false;
+		}
+		FNavigationToolViewModelPtr Parent = Item->GetParent();
 		while (Parent.IsValid())
 		{
 			if (SelectedItemSet.Contains(Parent))
@@ -964,7 +986,7 @@ void FNavigationTool::NormalizeToTopLevelSelections(TArray<FNavigationToolItemPt
 	});
 }
 
-void FNavigationTool::SyncSequencerSelection(const TArray<FNavigationToolItemPtr>& InSelectedItems) const
+void FNavigationTool::SyncSequencerSelection(const TArray<FNavigationToolViewModelWeakPtr>& InWeakSelectedItems) const
 {
 	const TSharedPtr<ISequencer> Sequencer = GetSequencer();
 	if (!Sequencer.IsValid())
@@ -973,9 +995,12 @@ void FNavigationTool::SyncSequencerSelection(const TArray<FNavigationToolItemPtr
 	}
 
 	FNavigationToolScopedSelection ScopedSelection(*Sequencer, ENavigationToolScopedSelectionPurpose::Sync);
-	for (const FNavigationToolItemPtr& Item : InSelectedItems)
+	for (const FNavigationToolViewModelWeakPtr& WeakItem : InWeakSelectedItems)
 	{
-		Item->Select(ScopedSelection);
+		if (const FNavigationToolViewModelPtr Item = WeakItem.Pin())
+		{
+			Item->Select(ScopedSelection);
+		}
 	}
 }
 
@@ -991,15 +1016,15 @@ FNavigationToolItemProxyRegistry& FNavigationTool::GetItemProxyRegistry()
 
 void FNavigationTool::OnObjectsReplaced(const TMap<UObject*, UObject*>& InReplacementMap)
 {
-	ForEachItem([&InReplacementMap](const FNavigationToolItemPtr& InItem)
+	ForEachItem([&InReplacementMap](const FNavigationToolViewModelPtr& InItem)
 		{
-			// Recursive not needed since we're calling it on all Items in Map Anyway
+			// Recursive not needed since we're calling it on all items in map anyway
 			InItem->OnObjectsReplaced(InReplacementMap, /*bRecursive*/false);
 		});
 
-	for (const TSharedPtr<INavigationToolAction>& Action : PendingActions)
+	for (const TSharedPtr<INavigationToolItemAction>& Action : PendingActions)
 	{
-		// Recursive needed since we only have direct reference to the underlying item in the Action, not its children
+		// Recursive needed since we only have direct reference to the underlying item in the action, not its children
 		if (Action.IsValid())
 		{
 			Action->OnObjectsReplaced(InReplacementMap, true);
@@ -1137,7 +1162,7 @@ void FNavigationTool::UnbindEvents()
 	FEditorDelegates::OnEditorActorReplaced.RemoveAll(this);
 }
 
-void FNavigationTool::AddItem(const FNavigationToolItemPtr& InItem)
+void FNavigationTool::AddItem(const FNavigationToolViewModelPtr& InItem)
 {
 	const FNavigationToolItemId ItemId = InItem->GetItemId();
 
@@ -1167,12 +1192,12 @@ void FNavigationTool::RemoveItem(const FNavigationToolItemId& InItemId)
 	}
 }
 
-void FNavigationTool::ForEachItem(TFunctionRef<void(const FNavigationToolItemPtr&)> InFunc)
+void FNavigationTool::ForEachItem(TFunctionRef<void(const FNavigationToolViewModelPtr&)> InFunc)
 {
 	// iteration scope, allowing for nested for-each
 	{
 		TGuardValue<bool> Guard(bIteratingItemMap, true);
-		for (const TPair<FNavigationToolItemId, FNavigationToolItemPtr>& Pair : ItemMap)
+		for (const TPair<FNavigationToolItemId, FNavigationToolViewModelPtr>& Pair : ItemMap)
 		{
 			InFunc(Pair.Value);
 		}
@@ -1180,7 +1205,7 @@ void FNavigationTool::ForEachItem(TFunctionRef<void(const FNavigationToolItemPtr
 
 	if (!bIteratingItemMap && (!ItemsPendingAdd.IsEmpty() || !ItemsPendingRemove.IsEmpty()))
 	{
-		for (const TPair<FNavigationToolItemId, FNavigationToolItemPtr>& ItemToAdd : ItemsPendingAdd)
+		for (const TPair<FNavigationToolItemId, FNavigationToolViewModelPtr>& ItemToAdd : ItemsPendingAdd)
 		{
 			ItemMap.Add(ItemToAdd);
 		}
@@ -1238,7 +1263,7 @@ void FNavigationTool::OnSequencerSelectionChanged()
 
 	if (!ItemsLastSelected.IsValid())
 	{
-		ItemsLastSelected = MakeShared<TArray<FNavigationToolItemWeakPtr>>();
+		ItemsLastSelected = MakeShared<TArray<FNavigationToolViewModelWeakPtr>>();
 	}
 
 	ItemsLastSelected->Reserve(ItemsLastSelected->Num() + SequencerSelection->Outliner.Num());
@@ -1249,12 +1274,12 @@ void FNavigationTool::OnSequencerSelectionChanged()
 	}
 }
 
-TArray<FNavigationToolItemPtr> FNavigationTool::FindItemsFromMovieSceneObject(UObject* const InObject) const
+TArray<FNavigationToolViewModelPtr> FNavigationTool::FindItemsFromMovieSceneObject(UObject* const InObject) const
 {
-	TArray<FNavigationToolItemPtr> OutItems;
+	TArray<FNavigationToolViewModelPtr> OutItems;
 	if (InObject)
 	{
-		for (const FNavigationToolItemPtr& Item : RootItem->GetDescendants())
+		for (const FNavigationToolViewModelPtr& Item : RootItem->GetDescendantsOfType<INavigationToolItem>())
 		{
 			if (Item->GetItemObject() == InObject)
 			{
@@ -1265,19 +1290,16 @@ TArray<FNavigationToolItemPtr> FNavigationTool::FindItemsFromMovieSceneObject(UO
 	return OutItems;
 }
 
-TArray<FNavigationToolItemPtr> FNavigationTool::FindItemsFromObjectGuid(const FGuid& InObjectGuid) const
+TArray<FNavigationToolViewModelPtr> FNavigationTool::FindItemsFromObjectGuid(const FGuid& InObjectGuid) const
 {
-	TArray<FNavigationToolItemPtr> OutItems;
+	TArray<FNavigationToolViewModelPtr> OutItems;
 	if (InObjectGuid.IsValid())
 	{
-		for (const FNavigationToolItemPtr& Item : RootItem->GetDescendants())
+		for (const TViewModelPtr<FNavigationToolBinding>& BindingItem : RootItem->GetDescendantsOfType<FNavigationToolBinding>())
 		{
-			if (const FNavigationToolBinding* const BindingItem = Item->CastTo<FNavigationToolBinding>())
+			if (BindingItem->GetBinding().GetObjectGuid() == InObjectGuid)
 			{
-				if (BindingItem->GetBinding().GetObjectGuid() == InObjectGuid)
-				{
-					OutItems.Add(Item);
-				}
+				OutItems.Add(BindingItem);
 			}
 		}
 	}
diff --git a/Engine/Plugins/Experimental/SequenceNavigator/Source/SequenceNavigator/Private/NavigationTool.h b/Engine/Plugins/Experimental/SequenceNavigator/Source/SequenceNavigator/Private/NavigationTool.h
--- a/Engine/Plugins/Experimental/SequenceNavigator/Source/SequenceNavigator/Private/NavigationTool.h
+++ b/Engine/Plugins/Experimental/SequenceNavigator/Source/SequenceNavigator/Private/NavigationTool.h
@@ -29,7 +29,7 @@ class FNavigationToolProvider;
 class FNavigationToolTab;
 class FNavigationToolTreeRoot;
 class FNavigationToolView;
-class INavigationToolAction;
+class INavigationToolItemAction;
 class INavigationToolView;
 enum class ENavigationToolProvidersChangeType;
 
@@ -49,7 +49,7 @@ public:
 	FOnToolLoaded OnToolLoaded;
 
 	/** Gathers all previously existing and new Item Proxies for a given Item */
-	void GetItemProxiesForItem(const FNavigationToolItemPtr& InItem, TArray<TSharedPtr<FNavigationToolItemProxy>>& OutItemProxies);
+	void GetItemProxiesForItem(const FNavigationToolViewModelPtr& InItem, TArray<TSharedPtr<FNavigationToolItemProxy>>& OutItemProxies);
 
 	/** Tries to find the Item Proxy Factory for the given Item Proxy Type Name */
 	INavigationToolItemProxyFactory* GetItemProxyFactory(const FName InItemProxyTypeName) const;
@@ -79,36 +79,36 @@ public:
 
 	virtual bool IsObjectAllowedInTool(const UObject* const InObject) const override;
 
-	virtual void RegisterItem(const FNavigationToolItemPtr& InItem) override;
+	virtual void RegisterItem(const FNavigationToolViewModelPtr& InItem) override;
 	virtual void UnregisterItem(const FNavigationToolItemId& InItemId) override;
 
 	virtual void RequestRefresh() override;
 	virtual void Refresh() override;
 
-	virtual TSharedRef<FNavigationToolItem> GetTreeRoot() const override;
+	virtual FNavigationToolViewModelWeakPtr GetTreeRoot() const override;
 
-	virtual FNavigationToolItemPtr FindItem(const FNavigationToolItemId& InItemId) const override;
-	virtual TArray<FNavigationToolItemPtr> TryFindItems(const UE::Sequencer::FViewModelPtr& InViewModel) const override;
+	virtual FNavigationToolViewModelPtr FindItem(const FNavigationToolItemId& InItemId) const override;
+	virtual TArray<FNavigationToolViewModelPtr> TryFindItems(const Sequencer::FViewModelPtr& InSequencerOutlinerViewModel) const override;
 
 	virtual void SetIgnoreNotify(const ENavigationToolIgnoreNotifyFlags InFlag, const bool bInIgnore) override;
 
 	virtual void OnSequencerSelectionChanged() override;
-	virtual TArray<FNavigationToolItemPtr> GetSelectedItems(const bool bInNormalizeToTopLevelSelections = false) const override;
-	virtual void SelectItems(const TArray<FNavigationToolItemPtr>& InItems
+	virtual TArray<FNavigationToolViewModelWeakPtr> GetSelectedItems(const bool bInNormalizeToTopLevelSelections = false) const override;
+	virtual void SelectItems(const TArray<FNavigationToolViewModelWeakPtr>& InWeakItems
 		, const ENavigationToolItemSelectionFlags InFlags = ENavigationToolItemSelectionFlags::SignalSelectionChange) const override;
 	virtual void ClearItemSelection(const bool bInSignalSelectionChange) const override;
 
 	virtual const FNavigationToolItemProxyRegistry& GetItemProxyRegistry() const override;
 	virtual TArray<FName> GetRegisteredItemProxyTypeNames() const override;
 
-	virtual void SetItemColor(const FNavigationToolItemPtr& InItem, const FColor& InColor) override;
-	virtual void RemoveItemColor(const FNavigationToolItemPtr& InItem) override;
-	virtual TOptional<FColor> FindItemColor(const FNavigationToolItemPtr& InItem, bool bRecurseParent = true) const override;
+	virtual void SetItemColor(const FNavigationToolViewModelPtr& InItem, const FColor& InColor) override;
+	virtual void RemoveItemColor(const FNavigationToolViewModelPtr& InItem) override;
+	virtual TOptional<FColor> FindItemColor(const FNavigationToolViewModelPtr& InItem, bool bRecurseParent = true) const override;
 
-	virtual void EnqueueItemActions(TArray<TSharedPtr<INavigationToolAction>>&& InItemActions) noexcept override;
+	virtual void EnqueueItemActions(TArray<TSharedPtr<INavigationToolItemAction>>&& InItemActions) noexcept override;
 
-	virtual void NotifyToolItemRenamed(const FNavigationToolItemPtr& InItem) override;
-	virtual void NotifyToolItemDeleted(const FNavigationToolItemPtr& InItem) override;
+	virtual void NotifyToolItemRenamed(const FNavigationToolViewModelPtr& InItem) override;
+	virtual void NotifyToolItemDeleted(const FNavigationToolViewModelPtr& InItem) override;
 
 	//~ End INavigationTool
 
@@ -125,7 +125,7 @@ public:
 	//~ End FTickableObjectBase
 
 	/** Delete a set of items in the outliner by calling their custom delete handler */
-	void DeleteItems(TArray<FNavigationToolItemPtr> InItems);
+	void DeleteItems(const TArray<FNavigationToolViewModelWeakPtr>& InWeakItems);
 
 	/** Unregisters the Navigation Tool View bound to the given id */
 	void UnregisterToolView(const int32 InToolViewId);
@@ -146,22 +146,22 @@ public:
 	 * Replaces the Item's Id in the Item Map. This can be due to an object item changing it's object
 	 * (e.g. a bp component getting destroyed and recreated, the item should be the same but the underlying component will not be)
 	 */
-	void NotifyItemIdChanged(const FNavigationToolItemId& OldId, const FNavigationToolItemPtr& InItem);
+	void NotifyItemIdChanged(const FNavigationToolItemId& OldId, const FNavigationToolViewModelPtr& InItem);
 
 	/** Gets the closest item to all the given items while also being their common ancestor */
-	static FNavigationToolItemPtr FindLowestCommonAncestor(const TArray<FNavigationToolItemPtr>& Items);
+	static FNavigationToolViewModelPtr FindLowestCommonAncestor(const TArray<FNavigationToolViewModelPtr>& Items);
 
 	/**
 	 * Sort the given array of items based on their ordering in the Navigation Tool
 	 * @see FNavigationTool::CompareToolItemOrder
 	 */
-	static void SortItems(TArray<FNavigationToolItemPtr>& OutToolItems, const bool bInReverseOrder = false);
+	static void SortItems(TArray<FNavigationToolViewModelWeakPtr>& WeakItems, const bool bInReverseOrder = false);
 
 	/** Normalizes the given Items by removing selected items that have their parent item also in the selection */
-	static void NormalizeToTopLevelSelections(TArray<FNavigationToolItemPtr>& InOutItems);
+	static void NormalizeToTopLevelSelections(TArray<FNavigationToolViewModelWeakPtr>& WeakItems);
 
 	/** Have the given Selected Items sync to the USelection Instances of Mode Tools */
-	void SyncSequencerSelection(const TArray<FNavigationToolItemPtr>& InSelectedItems) const;
+	void SyncSequencerSelection(const TArray<FNavigationToolViewModelWeakPtr>& InWeakSelectedItems) const;
 
 	/** Called when the engine replaces an object. A common example is when a BP Component is destroyed, and replaced */
 	void OnObjectsReplaced(const TMap<UObject*, UObject*>& InReplacementMap);
@@ -184,16 +184,16 @@ private:
 	void BindEvents();
 	void UnbindEvents();
 
-	void AddItem(const FNavigationToolItemPtr& InItem);
+	void AddItem(const FNavigationToolViewModelPtr& InItem);
 	void RemoveItem(const FNavigationToolItemId& InItemId);
-	void ForEachItem(TFunctionRef<void(const FNavigationToolItemPtr&)> InFunc);
+	void ForEachItem(TFunctionRef<void(const FNavigationToolViewModelPtr&)> InFunc);
 
 	USequencerSettings* GetSequencerSettings() const;
 
 	void OnTreeViewChanged();
 
-	TArray<FNavigationToolItemPtr> FindItemsFromMovieSceneObject(UObject* const InObject) const;
-	TArray<FNavigationToolItemPtr> FindItemsFromObjectGuid(const FGuid& InObjectGuid) const;
+	TArray<FNavigationToolViewModelPtr> FindItemsFromMovieSceneObject(UObject* const InObject) const;
+	TArray<FNavigationToolViewModelPtr> FindItemsFromObjectGuid(const FGuid& InObjectGuid) const;
 
 	bool AreAllViewsSyncingItemSelection() const;
 
@@ -206,14 +206,14 @@ private:
 	TSharedRef<FNavigationToolTab> ToolTab;
 
 	/** The root of all the items in the outliner */
-	TSharedRef<FNavigationToolTreeRoot> RootItem;
+	Sequencer::TViewModelPtr<FNavigationToolTreeRoot> RootItem;
 
 	TSharedPtr<FUICommandList> BaseCommandList;
 
 	/** The map of the registered items */
-	TMap<FNavigationToolItemId, FNavigationToolItemPtr> ItemMap;
+	TMap<FNavigationToolItemId, FNavigationToolViewModelPtr> ItemMap;
 
-	TMap<FNavigationToolItemId, FNavigationToolItemPtr> ItemsPendingAdd;
+	TMap<FNavigationToolItemId, FNavigationToolViewModelPtr> ItemsPendingAdd;
 
 	TSet<FNavigationToolItemId> ItemsPendingRemove;
 
@@ -221,10 +221,10 @@ private:
 	FNavigationToolItemProxyRegistry ItemProxyRegistry;
 	
 	/** The current pending actions before refresh is called */
-	TArray<TSharedPtr<INavigationToolAction>> PendingActions;
+	TArray<TSharedPtr<INavigationToolItemAction>> PendingActions;
 
 	/** The list of items pending selection processing, filled in when the Sequencer selection changes */
-	TSharedPtr<TArray<FNavigationToolItemWeakPtr>> ItemsLastSelected;
+	TSharedPtr<TArray<FNavigationToolViewModelWeakPtr>> ItemsLastSelected;
 
 	/** The map of registered outliner views */
 	TMap<int32, TSharedPtr<FNavigationToolView>> ToolViews;
diff --git a/Engine/Plugins/Experimental/SequenceNavigator/Source/SequenceNavigator/Private/NavigationToolExtender.cpp b/Engine/Plugins/Experimental/SequenceNavigator/Source/SequenceNavigator/Private/NavigationToolExtender.cpp
--- a/Engine/Plugins/Experimental/SequenceNavigator/Source/SequenceNavigator/Private/NavigationToolExtender.cpp
+++ b/Engine/Plugins/Experimental/SequenceNavigator/Source/SequenceNavigator/Private/NavigationToolExtender.cpp
@@ -496,7 +496,7 @@ bool FNavigationToolExtender::AnyProviderSupportsSequencer(const FName InToolId,
 	return bSupportsSequencer;
 }
 
-FSlateIcon FNavigationToolExtender::FindOverrideIcon(TSharedPtr<const FNavigationToolItem> InItem)
+FSlateIcon FNavigationToolExtender::FindOverrideIcon(const FNavigationToolViewModelPtr& InItem)
 {
 	if (const TSharedPtr<INavigationToolIconCustomization> CustomizationToUse = Get().GetCustomizationForItem(InItem))
 	{
@@ -505,10 +505,10 @@ FSlateIcon FNavigationToolExtender::FindOverrideIcon(TSharedPtr<const FNavigatio
 	return FSlateIcon();
 }
 
-void FNavigationToolExtender::RegisterOverridenIcon_Internal(const FNavigationToolItemTypeId& InItemTypeId, const TSharedRef<INavigationToolIconCustomization>& InIconCustomization)
+void FNavigationToolExtender::RegisterOverridenIcon_Internal(const Sequencer::FViewModelTypeID& InItemTypeId, const TSharedRef<INavigationToolIconCustomization>& InIconCustomization)
 {
 	FIconCustomizationKey Key;
-	Key.ItemTypeId = InItemTypeId;
+	Key.ItemTypeId = InItemTypeId.GetTypeID();
 	Key.CustomizationSpecializationIdentifier = InIconCustomization->GetToolItemIdentifier();
 
 	if (!IconRegistry.Contains(Key))
@@ -517,16 +517,17 @@ void FNavigationToolExtender::RegisterOverridenIcon_Internal(const FNavigationTo
 	}
 }
 
-void FNavigationToolExtender::UnregisterOverriddenIcon_Internal(const FNavigationToolItemTypeId& InItemTypeId, const FName& InSpecializationIdentifier)
+void FNavigationToolExtender::UnregisterOverriddenIcon_Internal(const Sequencer::FViewModelTypeID& InItemTypeId, const FName& InSpecializationIdentifier)
 {
 	FIconCustomizationKey Key;
-	Key.ItemTypeId = InItemTypeId;
+	Key.ItemTypeId = InItemTypeId.GetTypeID();
 	Key.CustomizationSpecializationIdentifier = InSpecializationIdentifier;
 
 	IconRegistry.Remove(Key);
 }
 
-TSharedPtr<INavigationToolIconCustomization> FNavigationToolExtender::GetCustomizationForItem(const TSharedPtr<const FNavigationToolItem>& InItem) const
+TSharedPtr<INavigationToolIconCustomization> FNavigationToolExtender::GetCustomizationForItem(
+	const FNavigationToolViewModelPtr& InItem) const
 {
 	if (!InItem.IsValid())
 	{
@@ -536,7 +537,7 @@ TSharedPtr<INavigationToolIconCustomization> FNavigationToolExtender::GetCustomi
 	// There is always going to be just one icon customization that supports the item since it checks specifically for the SupportCustomization
 	for (const TPair<FIconCustomizationKey, TSharedPtr<INavigationToolIconCustomization>>& IconCustomization : IconRegistry)
 	{
-		if (IconCustomization.Key.ItemTypeId == InItem->GetTypeId()
+		if (IconCustomization.Key.ItemTypeId == InItem.AsModel()->GetTypeTable().GetTypeID()
 			&& IconCustomization.Value->HasOverrideIcon(InItem))
 		{
 			return IconCustomization.Value;
diff --git a/Engine/Plugins/Experimental/SequenceNavigator/Source/SequenceNavigator/Private/NavigationToolScopedSelection.cpp b/Engine/Plugins/Experimental/SequenceNavigator/Source/SequenceNavigator/Private/NavigationToolScopedSelection.cpp
--- a/Engine/Plugins/Experimental/SequenceNavigator/Source/SequenceNavigator/Private/NavigationToolScopedSelection.cpp
+++ b/Engine/Plugins/Experimental/SequenceNavigator/Source/SequenceNavigator/Private/NavigationToolScopedSelection.cpp
@@ -12,13 +12,12 @@
 #include "MVVM/ViewModels/TrackModel.h"
 #include "MVVM/Views/SOutlinerView.h"
 #include "NavigationToolSettings.h"
-#include "Utils/NavigationToolMiscUtils.h"
-
-using namespace UE::Sequencer;
 
 namespace UE::SequenceNavigator
 {
 
+using namespace Sequencer;
+
 FNavigationToolScopedSelection::FNavigationToolScopedSelection(ISequencer& InSequencer
 	, ENavigationToolScopedSelectionPurpose InPurpose)
 	: Sequencer(InSequencer)
