[SequenceNavigator] Refactor to use Sequencer view models and type macros

diff --git a/Engine/Plugins/Experimental/SequenceNavigator/Source/SequenceNavigator/Private/NavigationToolView.cpp b/Engine/Plugins/Experimental/SequenceNavigator/Source/SequenceNavigator/Private/NavigationToolView.cpp
--- a/Engine/Plugins/Experimental/SequenceNavigator/Source/SequenceNavigator/Private/NavigationToolView.cpp
+++ b/Engine/Plugins/Experimental/SequenceNavigator/Source/SequenceNavigator/Private/NavigationToolView.cpp
@@ -39,6 +39,8 @@
 namespace UE::SequenceNavigator
 {
 
+using namespace Sequencer;
+
 FNavigationToolView::FNavigationToolView(FPrivateToken)
 	: ToolSettings(GetMutableDefault<UNavigationToolSettings>())
 	, ItemContextMenu(MakeShared<FNavigationToolItemContextMenu>())
@@ -137,12 +139,12 @@ void FNavigationToolView::PostLoad()
 	}
 }
 
-void FNavigationToolView::OnToolSettingsChanged(UObject* InObject, FPropertyChangedEvent& InPropertyChangedEvent)
+void FNavigationToolView::OnToolSettingsChanged(UObject* const InObject, FPropertyChangedEvent& InPropertyChangedEvent)
 {
 	RefreshTool(false);
 }
 
-void FNavigationToolView::Tick(float InDeltaTime)
+void FNavigationToolView::Tick(const float InDeltaTime)
 {
 	if (bRefreshRequested)
 	{
@@ -159,12 +161,12 @@ void FNavigationToolView::Tick(float InDeltaTime)
 	}
 
 	// Check if we have pending items to rename and we are not currently renaming an item
-	if (bRenamingItems && ItemsRemainingRename.Num() > 0 && !CurrentItemRenaming.IsValid())
+	if (bRenamingItems && WeakItemsRemainingRename.Num() > 0 && !WeakCurrentItemRenaming.Pin().IsValid())
 	{
-		CurrentItemRenaming = ItemsRemainingRename[0];
-		ItemsRemainingRename.RemoveAt(0);
+		WeakCurrentItemRenaming = WeakItemsRemainingRename[0];
+		WeakItemsRemainingRename.RemoveAt(0);
 
-		if (CurrentItemRenaming.IsValid())
+		if (const FNavigationToolViewModelPtr CurrentItemRenaming = WeakCurrentItemRenaming.Pin())
 		{
 			CurrentItemRenaming->OnRenameAction().AddSP(this, &FNavigationToolView::OnItemRenameAction);
 			CurrentItemRenaming->OnRenameAction().Broadcast(ENavigationToolRenameAction::Requested, SharedThis(this));
@@ -318,7 +320,7 @@ TSharedPtr<SWidget> FNavigationToolView::GetToolWidget() const
 
 TSharedPtr<SWidget> FNavigationToolView::CreateItemContextMenu()
 {
-	return ItemContextMenu->CreateMenu(SharedThis(this), SelectedItems);
+	return ItemContextMenu->CreateMenu(SharedThis(this), WeakSelectedItems);
 }
 
 bool FNavigationToolView::ShouldShowColumnByDefault(const TSharedPtr<INavigationToolColumn>& InColumn) const
@@ -376,38 +378,76 @@ void FNavigationToolView::SetKeyboardFocus()
 
 void FNavigationToolView::UpdateRootVisibleItems()
 {
-	RootVisibleItems.Reset();
-	ReadOnlyItems.Reset();
+	WeakRootVisibleItems.Reset();
+	WeakReadOnlyItems.Reset();
 	
 	if (WeakTool.IsValid())
 	{
-		GetChildrenOfItem(WeakTool.Pin()->GetTreeRoot(), RootVisibleItems);
+		GetChildrenOfItem(WeakTool.Pin()->GetTreeRoot(), WeakRootVisibleItems);
 	}
 }
 
 void FNavigationToolView::UpdateItemExpansions()
 {
-	TArray<FNavigationToolItemPtr> Items = RootVisibleItems;
+	for (const FNavigationToolViewModelWeakPtr& WeakRootVisibleItem : WeakRootVisibleItems)
+	{
+		const FNavigationToolViewModelPtr RootVisibleItem = WeakRootVisibleItem.Pin();
+		if (!RootVisibleItem.IsValid())
+		{
+			continue;
+		}
 
-	while (Items.Num() > 0)
+		for (const FNavigationToolViewModelPtr& Item : RootVisibleItem.AsModel()->GetDescendantsOfType<INavigationToolItem>())
+		{
+			const ENavigationToolItemFlags ItemFlags = GetViewItemFlags(Item);
+			const bool bHasExpandedFlag = EnumHasAnyFlags(ItemFlags, ENavigationToolItemFlags::Expanded);
+			SetItemExpansion(Item, bHasExpandedFlag);
+		}
+	}
+
+	if (ToolViewWidget.IsValid())
 	{
-		FNavigationToolItemPtr Item = Items.Pop();
-		const ENavigationToolItemFlags ItemFlags = GetViewItemFlags(Item);
-		SetItemExpansion(Item, EnumHasAnyFlags(ItemFlags, ENavigationToolItemFlags::Expanded));
-		Items.Append(Item->GetChildren());
+		for (const FNavigationToolViewModelWeakPtr& WeakRootVisibleItem : WeakRootVisibleItems)
+		{
+			const FNavigationToolViewModelPtr RootVisibleItem = WeakRootVisibleItem.Pin();
+			if (!RootVisibleItem.IsValid())
+			{
+				continue;
+			}
+
+			for (const FNavigationToolViewModelPtr& Item : RootVisibleItem.AsModel()->GetDescendantsOfType<INavigationToolItem>())
+			{
+				ToolViewWidget->UpdateItemExpansions(Item);
+			}
+		}
 	}
 
-	Items = RootVisibleItems;
+	/*TArray<FNavigationToolViewModelWeakPtr> Items = WeakRootVisibleItems;
 
 	while (Items.Num() > 0)
 	{
-		const FNavigationToolItemPtr Item = Items.Pop();
-		if (ToolViewWidget.IsValid())
+		if (const FNavigationToolViewModelPtr Item = Items.Pop().Pin())
 		{
-			ToolViewWidget->UpdateItemExpansions(Item);
+			const ENavigationToolItemFlags ItemFlags = GetViewItemFlags(Item);
+			const bool bHasExpandedFlag = EnumHasAnyFlags(ItemFlags, ENavigationToolItemFlags::Expanded);
+			SetItemExpansion(Item, bHasExpandedFlag);
+			Items.Append(Item->GetChildren());
 		}
-		Items.Append(Item->GetChildren());
 	}
+
+	Items = WeakRootVisibleItems;
+
+	while (Items.Num() > 0)
+	{
+		if (const FNavigationToolViewModelPtr Item = Items.Pop().Pin())
+		{
+			if (ToolViewWidget.IsValid())
+			{
+				ToolViewWidget->UpdateItemExpansions(Item);
+			}
+			Items.Append(Item->GetChildren());
+		}
+	}*/
 }
 
 void FNavigationToolView::NotifyObjectsReplaced()
@@ -418,28 +458,30 @@ void FNavigationToolView::NotifyObjectsReplaced()
 	}
 }
 
-FNavigationToolItemPtr FNavigationToolView::GetRootItem() const
+FNavigationToolViewModelPtr FNavigationToolView::GetRootItem() const
 {
 	if (WeakTool.IsValid())
 	{
-		return WeakTool.Pin()->GetTreeRoot();
+		return WeakTool.Pin()->GetTreeRoot().Pin();
 	}
 	return nullptr;
 }
 
-const TArray<FNavigationToolItemPtr>& FNavigationToolView::GetRootVisibleItems() const
+const TArray<FNavigationToolViewModelWeakPtr>& FNavigationToolView::GetRootVisibleItems() const
 {
-	return RootVisibleItems;
+	return WeakRootVisibleItems;
 }
 
-void FNavigationToolView::SaveViewItemFlags(const FNavigationToolItemPtr& InItem, ENavigationToolItemFlags InFlags)
+void FNavigationToolView::SaveViewItemFlags(const FNavigationToolViewModelWeakPtr& InWeakItem
+	, const ENavigationToolItemFlags InFlags)
 {
-	if (!InItem.IsValid())
+	const FNavigationToolViewModelPtr Item = InWeakItem.Pin();
+	if (!Item.IsValid())
 	{
 		return;
 	}
 
-	const TSharedPtr<FNavigationToolProvider> Provider = InItem->GetProvider();
+	const TSharedPtr<FNavigationToolProvider> Provider = Item->GetProvider();
 	if (!Provider.IsValid())
 	{
 		return;
@@ -466,22 +508,23 @@ void FNavigationToolView::SaveViewItemFlags(const FNavigationToolItemPtr& InItem
 
 	if (InFlags == ENavigationToolItemFlags::None)
 	{
-		SaveState->ToolViewSaveStates[ToolViewId].ViewItemFlags.Remove(InItem->GetItemId().GetStringId());
+		SaveState->ToolViewSaveStates[ToolViewId].ViewItemFlags.Remove(Item->GetItemId().GetStringId());
 	}
 	else
 	{
-		SaveState->ToolViewSaveStates[ToolViewId].ViewItemFlags.Add(InItem->GetItemId().GetStringId(), InFlags);
+		SaveState->ToolViewSaveStates[ToolViewId].ViewItemFlags.Add(Item->GetItemId().GetStringId(), InFlags);
 	}
 }
 
-ENavigationToolItemFlags FNavigationToolView::GetViewItemFlags(const FNavigationToolItemPtr& InItem) const
+ENavigationToolItemFlags FNavigationToolView::GetViewItemFlags(const FNavigationToolViewModelWeakPtr& InWeakItem) const
 {
-	if (!InItem.IsValid())
+	const FNavigationToolViewModelPtr Item = InWeakItem.Pin();
+	if (!Item.IsValid())
 	{
 		return ENavigationToolItemFlags::None;
 	}
 
-	const TSharedPtr<FNavigationToolProvider> Provider = InItem->GetProvider();
+	const TSharedPtr<FNavigationToolProvider> Provider = Item->GetProvider();
 	if (!Provider.IsValid())
 	{
 		return ENavigationToolItemFlags::None;
@@ -499,7 +542,7 @@ ENavigationToolItemFlags FNavigationToolView::GetViewItemFlags(const FNavigation
 		return ENavigationToolItemFlags::None;
 	}
 
-	if (const ENavigationToolItemFlags* const OverrideFlags = ViewSaveState->ViewItemFlags.Find(InItem->GetItemId().GetStringId()))
+	if (const ENavigationToolItemFlags* const OverrideFlags = ViewSaveState->ViewItemFlags.Find(Item->GetItemId().GetStringId()))
 	{
 		return *OverrideFlags;
 	}
@@ -507,49 +550,52 @@ ENavigationToolItemFlags FNavigationToolView::GetViewItemFlags(const FNavigation
 	return ENavigationToolItemFlags::None;
 }
 
-void FNavigationToolView::GetChildrenOfItem(const FNavigationToolItemPtr InItem, TArray<FNavigationToolItemPtr>& OutChildren) const
+void FNavigationToolView::GetChildrenOfItem(const FNavigationToolViewModelWeakPtr InWeakItem
+	, TArray<FNavigationToolViewModelWeakPtr>& OutWeakChildren) const
 {
-	static const TSet<FNavigationToolItemPtr> EmptySet;
-	GetChildrenOfItem(InItem, OutChildren, ENavigationToolItemViewMode::ItemTree, EmptySet);
+	static const TSet<FNavigationToolViewModelWeakPtr> EmptySet;
+	GetChildrenOfItem(InWeakItem, OutWeakChildren, ENavigationToolItemViewMode::ItemTree, EmptySet);
 }
 
-void FNavigationToolView::GetChildrenOfItem(const FNavigationToolItemPtr& InItem
-	, TArray<FNavigationToolItemPtr>& OutChildren
-	, ENavigationToolItemViewMode InViewMode
-	, const TSet<FNavigationToolItemPtr>& InRecursionDisallowedItems) const
+void FNavigationToolView::GetChildrenOfItem(const FNavigationToolViewModelWeakPtr InWeakItem
+	, TArray<FNavigationToolViewModelWeakPtr>& OutWeakChildren
+	, const ENavigationToolItemViewMode InViewMode
+	, const TSet<FNavigationToolViewModelWeakPtr>& InWeakRecursionDisallowedItems) const
 {
-	if (!InItem.IsValid())
+	const FNavigationToolViewModelPtr Item = InWeakItem.Pin();
+	if (!Item.IsValid())
 	{
 		return;
 	}
 
-	const TSharedPtr<FNavigationToolProvider> Provider = InItem->GetProvider();
-	if (!Provider.IsValid() && InItem->GetItemId() != FNavigationToolItemId::RootId)
+	const TSharedPtr<FNavigationToolProvider> Provider = Item->GetProvider();
+	if (!Provider.IsValid() && Item->GetItemId() != FNavigationToolItemId::RootId)
 	{
-		UE_LOG(LogTemp, Warning, TEXT("Sequence Navigator Item Id \"%s\" has no provider, but is a root item!"), *InItem->GetItemId().GetStringId());
+		UE_LOG(LogTemp, Warning, TEXT("Sequence Navigator Item Id \"%s\" has no provider, but is a root item!")
+			, *Item->GetItemId().GetStringId());
 		return;
 	}
 
-	for (const FNavigationToolItemPtr& ChildItem : InItem->GetChildren())
+	for (const FNavigationToolViewModelWeakPtr& WeakChildItem : Item->GetChildren())
 	{
-		if (!ChildItem.IsValid())
+		if (!WeakChildItem.Pin().IsValid())
 		{
 			continue;
 		}
 
-		if (const bool bShouldShowItem = ShouldShowItem(ChildItem, true, InViewMode))
+		if (ShouldShowItem(WeakChildItem, true, InViewMode))
 		{
 			// If the current item is visible in outliner, add it to the children
-			OutChildren.Add(ChildItem);
+			OutWeakChildren.Add(WeakChildItem);
 		}
-		else if (!InRecursionDisallowedItems.Contains(ChildItem))
+		else if (!InWeakRecursionDisallowedItems.Contains(WeakChildItem))
 		{
-			TArray<FNavigationToolItemPtr> GrandChildren;
+			TArray<FNavigationToolViewModelWeakPtr> WeakGrandChildren;
 
 			// For Muted Hierarchy to be in effect, not only does it have to be on
 			// but also the item should be shown (without counting the filter pass)
 			const bool bShouldUseMutedHierarchy = ToolSettings->ShouldUseMutedHierarchy();
-			const bool bShouldShowItemWithoutFilters = ShouldShowItem(ChildItem, false, InViewMode);
+			const bool bShouldShowItemWithoutFilters = ShouldShowItem(WeakChildItem, false, InViewMode);
 			const bool bShouldMuteItem = bShouldUseMutedHierarchy && bShouldShowItemWithoutFilters;
 
 			// If Muted Hierarchy, there might be ONLY grand children that are just visible in other view modes, 
@@ -557,28 +603,28 @@ void FNavigationToolView::GetChildrenOfItem(const FNavigationToolItemPtr& InItem
 			// If it's NOT muted hierarchy, just get the grand children visible in the requested view mode, as this ChildItem is guaranteed to be hidden
 			const ENavigationToolItemViewMode ViewModeToUse = bShouldMuteItem ? ENavigationToolItemViewMode::All : InViewMode;
 
-			GetChildrenOfItem(ChildItem, GrandChildren, ViewModeToUse, InRecursionDisallowedItems);
+			GetChildrenOfItem(WeakChildItem, WeakGrandChildren, ViewModeToUse, InWeakRecursionDisallowedItems);
 
-			if (!GrandChildren.IsEmpty())
+			if (!WeakGrandChildren.IsEmpty())
 			{
 				if (bShouldMuteItem)
 				{
-					ReadOnlyItems.Add(ChildItem);	
-					OutChildren.Add(ChildItem);
+					WeakReadOnlyItems.Add(WeakChildItem);	
+					OutWeakChildren.Add(WeakChildItem);
 				}
 				else
 				{
 					// We can append them knowing that the ViewMode to use is the one passed in and there's no
 					// child that leaked from another view mode
 					ensure(ViewModeToUse == InViewMode);
-					OutChildren.Append(GrandChildren);
+					OutWeakChildren.Append(WeakGrandChildren);
 				}
 			}
 		}
 	}
 }
 
-FLinearColor FNavigationToolView::GetItemBrushColor(const FNavigationToolItemPtr InItem) const
+FLinearColor FNavigationToolView::GetItemBrushColor(const FNavigationToolViewModelPtr InItem) const
 {
 	if (InItem.IsValid())
 	{
@@ -586,8 +632,8 @@ FLinearColor FNavigationToolView::GetItemBrushColor(const FNavigationToolItemPtr
 
 		// If NextSelectedItemIntoView is valid, it means we're scrolling items into view with Next/Previous, 
 		// so Make everything that's not the Current Item a bit more translucent to make the Current Item stand out
-		if (SortedSelectedItems.IsValidIndex(NextSelectedItemIntoView)
-			&& SortedSelectedItems[NextSelectedItemIntoView] != InItem)
+		if (WeakSortedSelectedItems.IsValidIndex(NextSelectedItemIntoView)
+			&& WeakSortedSelectedItems[NextSelectedItemIntoView] != InItem)
 		{
 			OutColor.A *= 0.5f;
 		}
@@ -598,44 +644,51 @@ FLinearColor FNavigationToolView::GetItemBrushColor(const FNavigationToolItemPtr
 	return FStyleColors::White.GetSpecifiedColor();
 }
 
-TArray<FNavigationToolItemPtr> FNavigationToolView::GetSelectedItems() const
+TArray<FNavigationToolViewModelWeakPtr> FNavigationToolView::GetSelectedItems() const
 {
-	return SelectedItems;
+	return WeakSelectedItems;
 }
 
 int32 FNavigationToolView::GetViewSelectedItemCount() const
 {
-	return SelectedItems.Num();
+	return WeakSelectedItems.Num();
 }
 
 int32 FNavigationToolView::CalculateVisibleItemCount() const
 {
-	TArray<FNavigationToolItemPtr> RemainingItems = RootVisibleItems;
-	
-	int32 VisibleItemCount = RemainingItems.Num();
+	TArray<FNavigationToolViewModelWeakPtr> WeakRemainingItems = WeakRootVisibleItems;
 
-	while (RemainingItems.Num() > 0)
+	int32 VisibleItemCount = WeakRemainingItems.Num();
+
+	while (WeakRemainingItems.Num() > 0)
 	{
-		const FNavigationToolItemPtr Item = RemainingItems.Pop();
-		
-		TArray<FNavigationToolItemPtr> ChildItems;
-		GetChildrenOfItem(Item, ChildItems);
-		VisibleItemCount += ChildItems.Num();
-		RemainingItems.Append(MoveTemp(ChildItems));
+		if (const FNavigationToolViewModelPtr Item = WeakRemainingItems.Pop().Pin())
+		{
+			TArray<FNavigationToolViewModelWeakPtr> WeakChildItems;
+			GetChildrenOfItem(Item, WeakChildItems);
+
+			VisibleItemCount += WeakChildItems.Num();
+
+			WeakRemainingItems.Append(MoveTemp(WeakChildItems));
+		}
 	}
 
-	// Remove the Read Only Items as they are filtered out items that are still shown because of Hierarchy Viz
-	VisibleItemCount -= ReadOnlyItems.Num();
+	// Remove the read only items as they are filtered out items that are still shown because of hierarchy viz
+	VisibleItemCount -= WeakReadOnlyItems.Num();
 
 	return VisibleItemCount;
 }
 
-void FNavigationToolView::SelectItems(TArray<FNavigationToolItemPtr> InItems, ENavigationToolItemSelectionFlags InFlags)
+void FNavigationToolView::SelectItems(TArray<FNavigationToolViewModelWeakPtr> InWeakItems
+	, const ENavigationToolItemSelectionFlags InFlags)
 {
-	// Remove Duplicate Items
-	TSet<FNavigationToolItemPtr> SeenItems;
-	SeenItems.Reserve(InItems.Num());
-	for (TArray<FNavigationToolItemPtr>::TIterator Iter(InItems); Iter; ++Iter)
+	const TSet<FNavigationToolViewModelWeakPtr> WeakItemSet(InWeakItems);
+
+	TSet<FNavigationToolViewModelWeakPtr> WeakSeenItems;
+	WeakSeenItems.Reserve(InWeakItems.Num());
+
+	// Remove duplicate items
+	/*for (TSet<FNavigationToolViewModelWeakPtr>::TIterator Iter(WeakItemSet); Iter; ++Iter)
 	{
 		if (SeenItems.Contains(*Iter))
 		{
@@ -645,21 +698,21 @@ void FNavigationToolView::SelectItems(TArray<FNavigationToolItemPtr> InItems, EN
 		{
 			SeenItems.Add(*Iter);
 		}
-	}
+	}*/
 
-	// Add the Children of the Items given
+	// Add the children of the items given
 	if (EnumHasAnyFlags(InFlags, ENavigationToolItemSelectionFlags::IncludeChildren))
 	{
-		TArray<FNavigationToolItemPtr> ChildItemsRemaining(InItems);
-		while (ChildItemsRemaining.Num() > 0)
+		TArray<FNavigationToolViewModelWeakPtr> WeakRemainingChildItems(InWeakItems);
+		while (WeakRemainingChildItems.Num() > 0)
 		{
-			if (FNavigationToolItemPtr ChildItem = ChildItemsRemaining.Pop())
+			if (FNavigationToolViewModelPtr ChildItem = WeakRemainingChildItems.Pop().Pin())
 			{
-				TArray<FNavigationToolItemPtr> Children;
-				GetChildrenOfItem(ChildItem, Children);
+				TArray<FNavigationToolViewModelWeakPtr> WeakChildren;
+				GetChildrenOfItem(ChildItem, WeakChildren);
 
-				InItems.Append(Children);
-				ChildItemsRemaining.Append(Children);
+				InWeakItems.Append(WeakChildren);
+				WeakRemainingChildItems.Append(WeakChildren);
 			}
 		}
 	}
@@ -667,56 +720,56 @@ void FNavigationToolView::SelectItems(TArray<FNavigationToolItemPtr> InItems, EN
 	if (EnumHasAnyFlags(InFlags, ENavigationToolItemSelectionFlags::AppendToCurrentSelection))
 	{
 		// Remove all repeated items to avoid duplicated entries
-		SelectedItems.RemoveAll([&SeenItems](const FNavigationToolItemPtr& Item)
-		{
-			return SeenItems.Contains(Item);
-		});	
-		TArray<FNavigationToolItemPtr> Items = MoveTemp(SelectedItems);
-		Items.Append(MoveTemp(InItems));
-		InItems = MoveTemp(Items);
+		WeakSelectedItems.RemoveAll([&WeakSeenItems](const FNavigationToolViewModelWeakPtr& InWeakItem)
+			{
+				return WeakSeenItems.Contains(InWeakItem);
+			});
+		TArray<FNavigationToolViewModelWeakPtr> Items = MoveTemp(WeakSelectedItems);
+		Items.Append(MoveTemp(InWeakItems));
+		InWeakItems = MoveTemp(Items);
 	}
 
-	if (!InItems.IsEmpty() && EnumHasAnyFlags(InFlags, ENavigationToolItemSelectionFlags::ScrollIntoView))
+	if (!InWeakItems.IsEmpty() && EnumHasAnyFlags(InFlags, ENavigationToolItemSelectionFlags::ScrollIntoView))
 	{
-		ScrollItemIntoView(InItems[0]);
+		ScrollItemIntoView(InWeakItems[0]);
 	}
 
 	const bool bSignalSelectionChange = EnumHasAnyFlags(InFlags, ENavigationToolItemSelectionFlags::SignalSelectionChange);
-	SetItemSelectionImpl(MoveTemp(InItems), bSignalSelectionChange);
+	SetItemSelectionImpl(MoveTemp(InWeakItems), bSignalSelectionChange);
 }
 
-void FNavigationToolView::ClearItemSelection(bool bSignalSelectionChange)
+void FNavigationToolView::ClearItemSelection(const bool bInSignalSelectionChange)
 {
-	SetItemSelectionImpl({}, bSignalSelectionChange);
+	SetItemSelectionImpl({}, bInSignalSelectionChange);
 }
 
-void FNavigationToolView::SetItemSelectionImpl(TArray<FNavigationToolItemPtr>&& InItems, bool bSignalSelectionChange)
+void FNavigationToolView::SetItemSelectionImpl(TArray<FNavigationToolViewModelWeakPtr>&& InWeakItems, const bool bInSignalSelectionChange)
 {
 	/*if (ShouldAutoExpandToSelection())
 	{
-		for (const FNavigationToolItemPtr& Item : InItems)
+		for (const FNavigationToolViewModelPtr& Item : InItems)
 		{
 			SetParentItemExpansions(Item, true);
 		}
 	}*/
 
-	SelectedItems = MoveTemp(InItems);
+	WeakSelectedItems = MoveTemp(InWeakItems);
 
 	if (ToolViewWidget.IsValid())
 	{
-		ToolViewWidget->SetItemSelection(SelectedItems, bSignalSelectionChange);
+		ToolViewWidget->SetItemSelection(WeakSelectedItems, bInSignalSelectionChange);
 	}
-	else if (bSignalSelectionChange)
+	else if (bInSignalSelectionChange)
 	{
-		NotifyItemSelectionChanged(SelectedItems, nullptr, true);
+		NotifyItemSelectionChanged(WeakSelectedItems, nullptr, true);
 	}
-	
+
 	Refresh();
 }
 
-void FNavigationToolView::NotifyItemSelectionChanged(const TArray<FNavigationToolItemPtr>& InSelectedItems
-	, const FNavigationToolItemPtr& InItem
-	, bool bUpdateModeTools)
+void FNavigationToolView::NotifyItemSelectionChanged(const TArray<FNavigationToolViewModelWeakPtr>& InWeakSelectedItems
+	, const FNavigationToolViewModelPtr& InItem
+	, const bool bInUpdateModeTools)
 {
 	if (bSyncingItemSelection)
 	{
@@ -724,108 +777,114 @@ void FNavigationToolView::NotifyItemSelectionChanged(const TArray<FNavigationToo
 	}
 	TGuardValue<bool> Guard(bSyncingItemSelection, true);
 
-	SelectedItems            = InSelectedItems;
-	SortedSelectedItems      = InSelectedItems;
+	WeakSelectedItems = InWeakSelectedItems;
+	WeakSortedSelectedItems = InWeakSelectedItems;
 	NextSelectedItemIntoView = -1;
 
-	FNavigationTool::SortItems(SortedSelectedItems);
+	FNavigationTool::SortItems(WeakSortedSelectedItems);
 
-	//If we have Pending Items Remaining but we switched Selection via Navigation, treat it as "I want to rename this too"
-	if (bRenamingItems && InItem.IsValid() && InItem != CurrentItemRenaming)
+	// If we have pending items remaining but we switched selection via navigation, treat it as "I want to rename this too"
+	if (bRenamingItems && InItem.IsValid() && InItem != WeakCurrentItemRenaming)
 	{
 		bRequestedRename = true;
 	}
 
 	const TSharedPtr<FNavigationTool> Tool = WeakTool.Pin();
-	if (bUpdateModeTools && Tool.IsValid())
+	if (bInUpdateModeTools && Tool.IsValid())
 	{
 		if (ToolSettings->ShouldSyncSelectionToSequencer())
 		{
-			Tool->SyncSequencerSelection(SelectedItems);
+			Tool->SyncSequencerSelection(WeakSelectedItems);
 		}
 
-		Tool->SelectItems(SelectedItems, ENavigationToolItemSelectionFlags::None);
+		Tool->SelectItems(WeakSelectedItems, ENavigationToolItemSelectionFlags::None);
 	}
 }
 
-bool FNavigationToolView::IsItemReadOnly(const FNavigationToolItemPtr& InItem) const
+bool FNavigationToolView::IsItemReadOnly(const FNavigationToolViewModelWeakPtr& InWeakItem) const
 {
-	return ReadOnlyItems.Contains(InItem);;
+	return WeakReadOnlyItems.Contains(InWeakItem);
 }
 
-bool FNavigationToolView::CanSelectItem(const FNavigationToolItemPtr& InItem) const
+bool FNavigationToolView::CanSelectItem(const FNavigationToolViewModelWeakPtr& InWeakItem) const
 {
-	const bool bIsSelectable = InItem.IsValid() && InItem->IsSelectable();
-	return bIsSelectable && !IsItemReadOnly(InItem);
+	const FNavigationToolViewModelPtr Item = InWeakItem.Pin();
+	const bool bIsSelectable = Item.IsValid() && Item->IsSelectable();
+	return bIsSelectable && !IsItemReadOnly(InWeakItem);
 }
 
-bool FNavigationToolView::IsItemSelected(const FNavigationToolItemPtr& InItem) const
+bool FNavigationToolView::IsItemSelected(const FNavigationToolViewModelWeakPtr& InWeakItem) const
 {
-	return SelectedItems.Contains(InItem);
+	return WeakSelectedItems.Contains(InWeakItem);
 }
 
-bool FNavigationToolView::IsItemExpanded(const FNavigationToolItemPtr& InItem, const bool bInUseFilter) const
+bool FNavigationToolView::IsItemExpanded(const FNavigationToolViewModelWeakPtr& InWeakItem, const bool bInUseFilter) const
 {
-	// Don't continue if Item should be hidden in view.
-	// the tree view still calls OnItemExpansionChanged even if it doesn't contain the item
-	// so this preemptive check is needed
-	if (!ShouldShowItem(InItem, bInUseFilter, ENavigationToolItemViewMode::ItemTree))
+	// Don't continue if item should be hidden in view.
+	// The tree view still calls OnItemExpansionChanged even if it doesn't contain the item
+	// so this preemptive check is needed.
+	if (!ShouldShowItem(InWeakItem, bInUseFilter, ENavigationToolItemViewMode::ItemTree))
 	{
 		return false;
 	}
 
 	if (ToolViewWidget.IsValid())
 	{
-		return ToolViewWidget->IsItemExpanded(InItem);
+		return ToolViewWidget->IsItemExpanded(InWeakItem);
 	}
 
 	return false;
 }
 
-void FNavigationToolView::SetItemExpansion(const FNavigationToolItemPtr& InItem, const bool bInExpand, const bool bInUseFilter)
+void FNavigationToolView::SetItemExpansion(const FNavigationToolViewModelWeakPtr& InWeakItem
+	, const bool bInExpand, const bool bInUseFilter)
 {
-	// Don't continue if Item should be hidden in view.
-	// the tree view still calls OnItemExpansionChanged even if it doesn't contain the item
+	// Don't continue if the item should be hidden in view.
+	// The tree view still calls OnItemExpansionChanged even if it doesn't contain the item
 	// so this preemptive check is needed
-	if (!ShouldShowItem(InItem, bInUseFilter, ENavigationToolItemViewMode::ItemTree))
+	if (!ShouldShowItem(InWeakItem, bInUseFilter, ENavigationToolItemViewMode::ItemTree))
 	{
 		return;
 	}
 
 	if (ToolViewWidget.IsValid())
 	{
-		ToolViewWidget->SetItemExpansion(InItem, bInExpand);
+		ToolViewWidget->SetItemExpansion(InWeakItem, bInExpand);
 	}
 	else
 	{
-		OnItemExpansionChanged(InItem, bInExpand);
+		OnItemExpansionChanged(InWeakItem, bInExpand);
 	}
 }
 
-void FNavigationToolView::SetItemExpansionRecursive(const FNavigationToolItemPtr InItem, const bool bInExpand)
+void FNavigationToolView::SetItemExpansionRecursive(const FNavigationToolViewModelWeakPtr InWeakItem, const bool bInExpand)
 {
-	SetItemExpansion(InItem, bInExpand, false);
-	
-	for (const FNavigationToolItemPtr& Child : InItem->GetChildren())
+	const FNavigationToolViewModelPtr Item = InWeakItem.Pin();
+	if (!Item.IsValid())
 	{
-		if (Child.IsValid())
-		{
-			SetItemExpansionRecursive(Child, bInExpand);
-		}
+		return;
+	}
+
+	SetItemExpansion(InWeakItem, bInExpand, false);
+
+	for (const FNavigationToolViewModelWeakPtr& WeakChildItem : Item->GetChildren())
+	{
+		SetItemExpansionRecursive(WeakChildItem, bInExpand);
 	}
 }
 
-void FNavigationToolView::SetParentItemExpansions(const FNavigationToolItemPtr& InItem, const bool bInExpand)
+void FNavigationToolView::SetParentItemExpansions(const FNavigationToolViewModelWeakPtr& InWeakItem, const bool bInExpand)
 {
-	if (!InItem.IsValid())
+	const FNavigationToolViewModelPtr Item = InWeakItem.Pin();
+	if (!Item.IsValid())
 	{
 		return;
 	}
-	
-	TArray<FNavigationToolItemPtr> ItemsToExpand;
-	
+
+	TArray<FNavigationToolViewModelPtr> ItemsToExpand;
+
 	// Don't auto expand at all if there's a parent preventing it
-	FNavigationToolItemPtr ParentItem = InItem->GetParent();
+	FNavigationToolViewModelPtr ParentItem = Item->GetParent();
 	while (ParentItem.IsValid())
 	{
 		if (!ParentItem->CanAutoExpand())
@@ -836,15 +895,21 @@ void FNavigationToolView::SetParentItemExpansions(const FNavigationToolItemPtr&
 		ParentItem = ParentItem->GetParent();
 	}
 
-	for (const FNavigationToolItemPtr& Item : ItemsToExpand)
+	for (const FNavigationToolViewModelPtr& ItemToExpand : ItemsToExpand)
 	{
-		SetItemExpansion(Item, bInExpand);
+		SetItemExpansion(ItemToExpand, bInExpand);
 	}
 }
 
-void FNavigationToolView::OnItemExpansionChanged(const FNavigationToolItemPtr InItem, const bool bInIsExpanded)
+void FNavigationToolView::OnItemExpansionChanged(const FNavigationToolViewModelWeakPtr InWeakItem, const bool bInIsExpanded)
 {
-	const ENavigationToolItemFlags CurrentFlags = GetViewItemFlags(InItem);
+	const FNavigationToolViewModelPtr Item = InWeakItem.Pin();
+	if (!Item.IsValid())
+	{
+		return;
+	}
+
+	const ENavigationToolItemFlags CurrentFlags = GetViewItemFlags(InWeakItem);
 	
 	ENavigationToolItemFlags TargetFlags = CurrentFlags;
 	
@@ -857,23 +922,24 @@ void FNavigationToolView::OnItemExpansionChanged(const FNavigationToolItemPtr In
 		TargetFlags &= ~ENavigationToolItemFlags::Expanded;
 	}
 
-	SaveViewItemFlags(InItem, TargetFlags);
+	SaveViewItemFlags(InWeakItem, TargetFlags);
 
 	if (CurrentFlags != TargetFlags)
 	{
-		InItem->OnExpansionChanged().Broadcast(SharedThis(this), bInIsExpanded);
+		Item->OnExpansionChanged().Broadcast(SharedThis(this), bInIsExpanded);
 	}
 }
 
-bool FNavigationToolView::ShouldShowItem(const FNavigationToolItemPtr& InItem, bool bInUseFilters
-	, ENavigationToolItemViewMode InViewMode) const
+bool FNavigationToolView::ShouldShowItem(const FNavigationToolViewModelWeakPtr& InWeakItem
+	, const bool bInUseFilters, const ENavigationToolItemViewMode InViewMode) const
 {
-	if (!InItem.IsValid())
+	const FNavigationToolViewModelPtr Item = InWeakItem.Pin();
+	if (!Item.IsValid())
 	{
 		return false;
 	}
 
-	if (InItem->IsA<FNavigationToolTreeRoot>())
+	if (Item.AsModel()->IsA<FNavigationToolTreeRoot>())
 	{
 		return true;
 	}
@@ -884,13 +950,13 @@ bool FNavigationToolView::ShouldShowItem(const FNavigationToolItemPtr& InItem, b
 		return false;
 	}
 
-	if (!InItem->IsAllowedInTool())
+	if (!Item->IsAllowedInTool())
 	{
 		//UE_LOG(LogTemp, Warning, TEXT("Sequence Navigator Item Id \"%s\" is hidden: Not Allowed In Tool"), *InItem->GetItemId().GetStringId());
 		return false;
 	}
 
-	if (!InItem->IsViewModeSupported(InViewMode, *this))
+	if (!Item->IsViewModeSupported(InViewMode, *this))
 	{
 		//UE_LOG(LogTemp, Warning, TEXT("Sequence Navigator Item Id \"%s\" is hidden: View Mode Not Supported"), *InItem->GetItemId().GetStringId());
 		return false;
@@ -898,10 +964,10 @@ bool FNavigationToolView::ShouldShowItem(const FNavigationToolItemPtr& InItem, b
 
 	// Allow providers to determine whether the item should be hidden
 	bool bProviderShouldHideItem = false;
-	Tool->ForEachProvider([&bProviderShouldHideItem, &InItem]
+	Tool->ForEachProvider([&bProviderShouldHideItem, &Item]
 		(const TSharedRef<FNavigationToolProvider>& InToolProvider)
 		{
-			if (InToolProvider->ShouldHideItem(InItem))
+			if (InToolProvider->ShouldHideItem(Item))
 			{
 				bProviderShouldHideItem = true;
 				return false;
@@ -916,12 +982,12 @@ bool FNavigationToolView::ShouldShowItem(const FNavigationToolItemPtr& InItem, b
 
 	// Extra pass for Non-Item Proxies that are parented under an Item Proxy
 	// Hiding an Item Proxy Type should affect all the rest of the items below it
-	if (!InItem->IsA<FNavigationToolItemProxy>())
+	if (!Item.AsModel()->IsA<FNavigationToolItemProxy>())
 	{
-		FNavigationToolItemPtr ItemParent = InItem->GetParent();
+		FNavigationToolViewModelPtr ItemParent = Item->GetParent();
 		while (ItemParent.IsValid())
 		{
-			if (ItemParent->IsA<FNavigationToolItemProxy>())
+			if (ItemParent.AsModel()->IsA<FNavigationToolItemProxy>())
 			{
 				// Stop at the first Item Proxy parent found
 				break;
@@ -935,7 +1001,7 @@ bool FNavigationToolView::ShouldShowItem(const FNavigationToolItemPtr& InItem, b
 
 	for (const TSharedPtr<FNavigationToolBuiltInFilter>& GlobalFilter : Tool->GlobalFilters)
 	{
-		if (!GlobalFilter->IsActive() && GlobalFilter->PassesFilter(InItem))
+		if (!GlobalFilter->IsActive() && GlobalFilter->PassesFilter(Item))
 		{
 			bGlobalFilterOut = true;
 			break;
@@ -947,7 +1013,7 @@ bool FNavigationToolView::ShouldShowItem(const FNavigationToolItemPtr& InItem, b
 		return false;
 	}
 
-	if (bInUseFilters && FilterBar->GetFilterData().IsFilteredOut(InItem))
+	if (bInUseFilters && FilterBar->GetFilterData().IsFilteredOut(Item))
 	{
 		//UE_LOG(LogTemp, Warning, TEXT("Sequence Navigator Item Id \"%s\" is hidden: Filtered Out"), *InItem->GetItemId().GetStringId());
 		return false;
@@ -956,26 +1022,26 @@ bool FNavigationToolView::ShouldShowItem(const FNavigationToolItemPtr& InItem, b
 	return true;
 }
 
-int32 FNavigationToolView::GetVisibleChildIndex(const FNavigationToolItemPtr& InParentItem, const FNavigationToolItemPtr& InChildItem) const
+int32 FNavigationToolView::GetVisibleChildIndex(const FNavigationToolViewModelPtr& InParentItem, const FNavigationToolViewModelPtr& InChildItem) const
 {
 	if (InParentItem.IsValid())
 	{
-		TArray<FNavigationToolItemPtr> Children;
-		GetChildrenOfItem(InParentItem, Children);
-		return Children.Find(InChildItem);
+		TArray<FNavigationToolViewModelWeakPtr> WeakChildren;
+		GetChildrenOfItem(InParentItem, WeakChildren);
+		return WeakChildren.Find(InChildItem);
 	}
 	return INDEX_NONE;
 }
 
-FNavigationToolItemPtr FNavigationToolView::GetVisibleChildAt(const FNavigationToolItemPtr& InParentItem, int32 InChildIndex) const
+FNavigationToolViewModelPtr FNavigationToolView::GetVisibleChildAt(const FNavigationToolViewModelPtr& InParentItem, int32 InChildIndex) const
 {
 	if (InParentItem.IsValid())
 	{
-		TArray<FNavigationToolItemPtr> Children;
-		GetChildrenOfItem(InParentItem, Children);
-		if (Children.IsValidIndex(InChildIndex))
+		TArray<FNavigationToolViewModelWeakPtr> WeakChildren;
+		GetChildrenOfItem(InParentItem, WeakChildren);
+		if (WeakChildren.IsValidIndex(InChildIndex))
 		{
-			return Children[InChildIndex];
+			return WeakChildren[InChildIndex].Pin();
 		}
 	}
 	return nullptr;
@@ -1171,81 +1237,104 @@ bool FNavigationToolView::IsItemTypeHidden(const FName InItemTypeName) const
 	return HiddenItemTypes.Contains(InItemTypeName);
 }
 
-bool FNavigationToolView::IsItemTypeHidden(const FNavigationToolItemPtr& InItem) const
+bool FNavigationToolView::IsItemTypeHidden(const FNavigationToolViewModelPtr& InItem) const
 {
-	return IsItemTypeHidden(InItem->GetTypeId().ToName());
+	return IsItemTypeHidden(InItem.AsModel()->GetTypeTable().GetTypeName());
 }
 
-void FNavigationToolView::OnDragEnter(const FDragDropEvent& InDragDropEvent, const FNavigationToolItemPtr InTargetItem)
+void FNavigationToolView::OnDragEnter(const FDragDropEvent& InDragDropEvent
+	, const FNavigationToolViewModelWeakPtr InWeakTargetItem)
 {
-	if (!InTargetItem.IsValid() && WeakTool.IsValid())
+	bool bDragIntoTreeRoot = false;
+
+	const FNavigationToolViewModelPtr TargetItem = InWeakTargetItem.Pin();
+	if (!TargetItem.IsValid() && WeakTool.IsValid())
 	{
-		TSharedRef<FNavigationToolItem> TreeRoot = WeakTool.Pin()->GetTreeRoot();
+		const FNavigationToolViewModelPtr TreeRoot = WeakTool.Pin()->GetTreeRoot().Pin();
+		const bool bCanAcceptDrop = TreeRoot.IsValid()
+			&& TreeRoot->CanAcceptDrop(InDragDropEvent, EItemDropZone::OntoItem).IsSet();
+		SetDragIntoTreeRoot(bCanAcceptDrop);
+	}
+
+	SetDragIntoTreeRoot(bDragIntoTreeRoot);
+
+	/*const FNavigationToolViewModelPtr TargetItem = InWeakTargetItem.Pin();
+	if (!TargetItem.IsValid() && WeakTool.IsValid())
+	{
+		const TViewModelPtr<FNavigationToolTreeRoot> TreeRoot = WeakTool.Pin()->GetTreeRoot();
 		const bool bCanAcceptDrop = TreeRoot->CanAcceptDrop(InDragDropEvent, EItemDropZone::OntoItem).IsSet();
 		SetDragIntoTreeRoot(bCanAcceptDrop);
 	}
 	else
 	{
 		SetDragIntoTreeRoot(false);
-	}
+	}*/
 }
 
-void FNavigationToolView::OnDragLeave(const FDragDropEvent& InDragDropEvent, const FNavigationToolItemPtr InTargetItem)
+void FNavigationToolView::OnDragLeave(const FDragDropEvent& InDragDropEvent
+	, const FNavigationToolViewModelWeakPtr InWeakTargetItem)
 {
 	// If drag left an item, set the drag into tree root to false (this will set it back to false if a valid item receives DragEnter)
-	SetDragIntoTreeRoot(InTargetItem.IsValid());
+	SetDragIntoTreeRoot(InWeakTargetItem.Pin().IsValid());
 }
 
 FReply FNavigationToolView::OnDragDetected(const FGeometry& InGeometry
-	, const FPointerEvent& InMouseEvent
-	, const FNavigationToolItemPtr InTargetItem)
+	, const FPointerEvent& InPointerEvent
+	, const FNavigationToolViewModelWeakPtr InWeakTargetItem)
 {
 	if (!IsToolLocked())
 	{
-		// Only Select Target if it hasn't already been selected
-		if (!IsItemSelected(InTargetItem))
-		{
-			const ENavigationToolItemSelectionFlags SelectionFlags = InMouseEvent.IsControlDown()
-				? ENavigationToolItemSelectionFlags::AppendToCurrentSelection
-				: ENavigationToolItemSelectionFlags::None;
+		return FReply::Unhandled();
+	}
 
-			SelectItems({ InTargetItem }, SelectionFlags);
-		}
+	// Only select the target if it hasn't already been selected
+	if (!IsItemSelected(InWeakTargetItem.ImplicitPin()))
+	{
+		const ENavigationToolItemSelectionFlags SelectionFlags = InPointerEvent.IsControlDown()
+			? ENavigationToolItemSelectionFlags::AppendToCurrentSelection
+			: ENavigationToolItemSelectionFlags::None;
 
-		// Get all Selected Items that are in a state where they can be selected again (i.e. not Read Only)
-		TArray<FNavigationToolItemPtr> Items = GetSelectedItems();
-		Items.RemoveAll([this](const FNavigationToolItemPtr& InItem)->bool
+		SelectItems({ InWeakTargetItem.ImplicitPin() }, SelectionFlags);
+	}
+
+	// Get all selected items that are in a state where they can be selected again (i.e. not read only)
+	TArray<FNavigationToolViewModelWeakPtr> WeakItems = GetSelectedItems();
+	WeakItems.RemoveAll([this](const FNavigationToolViewModelWeakPtr& InWeakItem) -> bool
 		{
-			return !CanSelectItem(InItem);
+			return !CanSelectItem(InWeakItem);
 		});
 
-		if (Items.Num() > 0)
-		{
-			const ENavigationToolDragDropActionType ActionType = InMouseEvent.IsAltDown()
-				? ENavigationToolDragDropActionType::Copy
-				: ENavigationToolDragDropActionType::Move;
+	if (WeakItems.Num() > 0)
+	{
+		const ENavigationToolDragDropActionType ActionType = InPointerEvent.IsAltDown()
+			? ENavigationToolDragDropActionType::Copy
+			: ENavigationToolDragDropActionType::Move;
 
-			const TSharedRef<FNavigationToolItemDragDropOp> DragDropOp = FNavigationToolItemDragDropOp::New(Items, SharedThis(this), ActionType);
-			return FReply::Handled().BeginDragDrop(DragDropOp);
-		}
+		const TSharedRef<FNavigationToolItemDragDropOp> DragDropOp =
+			FNavigationToolItemDragDropOp::New(WeakItems, SharedThis(this), ActionType);
+		return FReply::Handled().BeginDragDrop(DragDropOp);
 	}
+
 	return FReply::Unhandled();
 }
 
-FReply FNavigationToolView::OnDrop(const FDragDropEvent& InDragDropEvent, EItemDropZone InDropZone, const FNavigationToolItemPtr InTargetItem)
+FReply FNavigationToolView::OnDrop(const FDragDropEvent& InDragDropEvent
+	, const EItemDropZone InDropZone
+	, const FNavigationToolViewModelWeakPtr InWeakTargetItem)
 {
 	SetDragIntoTreeRoot(false);
 
-	if (InTargetItem.IsValid())
+	const FNavigationToolViewModelPtr TargetItem = InWeakTargetItem.Pin();
+	if (TargetItem.IsValid())
 	{
-		return InTargetItem->AcceptDrop(InDragDropEvent, InDropZone);
+		return TargetItem->AcceptDrop(InDragDropEvent, InDropZone);
 	}
 
-	FNavigationToolItemPtr TreeRoot;
+	FNavigationToolViewModelPtr TreeRoot;
 
 	if (WeakTool.IsValid())
 	{
-		TreeRoot = WeakTool.Pin()->GetTreeRoot();
+		TreeRoot = WeakTool.Pin()->GetTreeRoot().Pin();
 	}
 
 	if (TreeRoot.IsValid() && TreeRoot->CanAcceptDrop(InDragDropEvent, EItemDropZone::OntoItem))
@@ -1257,17 +1346,18 @@ FReply FNavigationToolView::OnDrop(const FDragDropEvent& InDragDropEvent, EItemD
 }
 
 TOptional<EItemDropZone> FNavigationToolView::OnCanDrop(const FDragDropEvent& InDragDropEvent
-	, EItemDropZone InDropZone
-	, const FNavigationToolItemPtr InTargetItem) const
-{	
-	if (!IsToolLocked() && InTargetItem.IsValid() && CanSelectItem(InTargetItem))
+	, const EItemDropZone InDropZone
+	, const FNavigationToolViewModelWeakPtr InWeakTargetItem) const
+{
+	const FNavigationToolViewModelPtr TargetItem = InWeakTargetItem.Pin();
+	if (TargetItem.IsValid() && !IsToolLocked() && CanSelectItem(TargetItem))
 	{
-		return InTargetItem->CanAcceptDrop(InDragDropEvent, InDropZone);
+		return TargetItem->CanAcceptDrop(InDragDropEvent, InDropZone);
 	}
 	return TOptional<EItemDropZone>();
 }
 
-void FNavigationToolView::SetDragIntoTreeRoot(bool bInIsDraggingIntoTreeRoot)
+void FNavigationToolView::SetDragIntoTreeRoot(const bool bInIsDraggingIntoTreeRoot)
 {
 	if (ToolViewWidget.IsValid())
 	{
@@ -1277,50 +1367,57 @@ void FNavigationToolView::SetDragIntoTreeRoot(bool bInIsDraggingIntoTreeRoot)
 
 void FNavigationToolView::RenameSelected()
 {
-	if (const TSharedPtr<FNavigationTool> Tool = WeakTool.Pin())
+	const TSharedPtr<FNavigationTool> Tool = WeakTool.Pin();
+	if (!Tool.IsValid())
 	{
-		TArray<FNavigationToolItemPtr> Items = GetSelectedItems();
-		
-		if (Items.IsEmpty())
-		{
-			return;
-		}
-	
-		// Assume we have an item currently renaming
-		ResetRenaming();
+		return;
+	}
+
+	TArray<FNavigationToolViewModelWeakPtr> WeakItems = GetSelectedItems();
+
+	if (WeakItems.IsEmpty())
+	{
+		return;
+	}
+
+	// Assume we have an item currently renaming
+	ResetRenaming();
 
-		// Remove Items that are Invalid or can't be renamed
-		Items.RemoveAll([](const FNavigationToolItemPtr& InItem)
+	// Remove items that are invalid or can't be renamed
+	WeakItems.RemoveAll([](const FNavigationToolViewModelWeakPtr& InWeakItem)
+		{
+			const FNavigationToolViewModelPtr Item = InWeakItem.Pin();
+			if (!Item.IsValid())
 			{
-				if (InItem.IsValid())
-				{
-					if (const IRenameableExtension* const RenameableExtension = InItem->CastTo<IRenameableExtension>())
-					{
-						return !RenameableExtension->CanRename();
-					}
-				}
 				return true;
-			});
-		
-		ItemsRemainingRename = MoveTemp(Items);
+			}
 
-		if (ItemsRemainingRename.Num() > 0)
-		{
-			FNavigationTool::SortItems(ItemsRemainingRename);
-			bRenamingItems = true;
-		}
+			if (const TViewModelPtr<IRenameableExtension> RenameableExtension = Item.ImplicitCast())
+			{
+				return !RenameableExtension->CanRename();
+			}
+
+			return true;
+		});
+
+	WeakItemsRemainingRename = MoveTemp(WeakItems);
+
+	if (WeakItemsRemainingRename.Num() > 0)
+	{
+		FNavigationTool::SortItems(WeakItemsRemainingRename);
+		bRenamingItems = true;
 	}
 }
 
 void FNavigationToolView::ResetRenaming()
 {
-	if (CurrentItemRenaming.IsValid())
+	if (const FNavigationToolViewModelPtr CurrentItemRenaming = WeakCurrentItemRenaming.Pin())
 	{
 		CurrentItemRenaming->OnRenameAction().RemoveAll(this);
-		CurrentItemRenaming.Reset();
+		WeakCurrentItemRenaming = nullptr;
 	}
 
-	if (ItemsRemainingRename.IsEmpty())
+	if (WeakItemsRemainingRename.IsEmpty())
 	{
 		bRenamingItems = false;
 	}
@@ -1342,7 +1439,7 @@ void FNavigationToolView::OnItemRenameAction(const ENavigationToolRenameAction I
 		break;
 	
 	case ENavigationToolRenameAction::Cancelled:
-		ItemsRemainingRename.Reset();
+		WeakItemsRemainingRename.Reset();
 		ResetRenaming();
 		break;
 	
@@ -1357,14 +1454,15 @@ void FNavigationToolView::OnItemRenameAction(const ENavigationToolRenameAction I
 
 bool FNavigationToolView::CanRenameSelected() const
 {
-	for (const FNavigationToolItemPtr& Item : SelectedItems)
+	for (const FNavigationToolViewModelWeakPtr& WeakItem : WeakSelectedItems)
 	{
+		const FNavigationToolViewModelPtr Item = WeakItem.Pin();
 		if (!Item.IsValid())
 		{
 			return false;
 		}
 
-		if (const IRenameableExtension* const RenameableExtension = Item->CastTo<IRenameableExtension>())
+		if (const TViewModelPtr<IRenameableExtension> RenameableExtension = Item.ImplicitCast())
 		{
 			if (!RenameableExtension->CanRename())
 			{
@@ -1373,39 +1471,45 @@ bool FNavigationToolView::CanRenameSelected() const
 		}
 	}
 
-	return SelectedItems.Num() > 0;
+	return WeakSelectedItems.Num() > 0;
 }
 
 void FNavigationToolView::DeleteSelected()
 {
 	if (const TSharedPtr<FNavigationTool> Tool = WeakTool.Pin())
 	{
-		TArray<FNavigationToolItemPtr> Items = GetSelectedItems();
+		TArray<FNavigationToolViewModelWeakPtr> WeakItems = GetSelectedItems();
 
-		Items.RemoveAll([](const FNavigationToolItemPtr& InItem)
-		{
-			return !InItem.IsValid() || !InItem->CanDelete();
-		});
+		WeakItems.RemoveAll([](const FNavigationToolViewModelWeakPtr& InWeakItem)
+			{
+				if (const FNavigationToolViewModelPtr Item = InWeakItem.Pin())
+				{
+					return !Item.IsValid() || !Item->CanDelete();
+				}
+				return false;
+			});
 
-		if (Items.IsEmpty())
+		if (WeakItems.IsEmpty())
 		{
 			return;
 		}
 
-		Tool->DeleteItems(Items);
+		Tool->DeleteItems(WeakItems);
 	}
 }
 
 bool FNavigationToolView::CanDeleteSelected() const
 {
-	for (const FNavigationToolItemPtr& Item : GetSelectedItems())
+	for (const FNavigationToolViewModelWeakPtr& WeakItem : GetSelectedItems())
 	{
-		if (Item && Item->CanDelete())
+		if (const FNavigationToolViewModelPtr Item = WeakItem.Pin())
 		{
-			return true;
+			if (Item->CanDelete())
+			{
+				return true;
+			}
 		}
 	}
-
 	return false;
 }
 
@@ -1421,7 +1525,7 @@ void FNavigationToolView::DuplicateSelected()
 bool FNavigationToolView::CanDuplicateSelected() const
 {
 	// Disabled for now until (if) we decide to support this
-	/*for (const FNavigationToolItemPtr& Item : GetSelectedItems())
+	/*for (const FNavigationToolViewModelPtr& Item : GetSelectedItems())
 	{
 		if (Item && Item->CanDuplicate())
 		{
@@ -1431,27 +1535,31 @@ bool FNavigationToolView::CanDuplicateSelected() const
 	return false;
 }
 
-void FNavigationToolView::SelectChildren(bool bIsRecursive)
+void FNavigationToolView::SelectChildren(const bool bInIsRecursive)
 {
-	TArray<FNavigationToolItemPtr> ItemsToSelect;
-	TArray<FNavigationToolItemPtr> RemainingItems = GetSelectedItems();
+	TArray<FNavigationToolViewModelWeakPtr> WeakItemsToSelect;
+	TArray<FNavigationToolViewModelWeakPtr> WeakRemainingItems = GetSelectedItems();
 
-	while (RemainingItems.Num() > 0)
+	while (WeakRemainingItems.Num() > 0)
 	{
-		//Note: Pop here will affect order of Children in Selection
-		const FNavigationToolItemPtr ParentItem = RemainingItems.Pop();
+		// Note: Pop here will affect the order of children in selection
+		const FNavigationToolViewModelPtr ParentItem = WeakRemainingItems.Pop().Pin();
+		if (!ParentItem.IsValid())
+		{
+			continue;
+		}
 
-		TArray<FNavigationToolItemPtr> ChildItems;
-		GetChildrenOfItem(ParentItem, ChildItems);
-		if (bIsRecursive)
+		TArray<FNavigationToolViewModelWeakPtr> WeakChildItems;
+		GetChildrenOfItem(ParentItem, WeakChildItems);
+		if (bInIsRecursive)
 		{
-			RemainingItems.Append(ChildItems);
+			WeakRemainingItems.Append(WeakChildItems);
 		}
-		ItemsToSelect.Append(ChildItems);
+		WeakItemsToSelect.Append(WeakChildItems);
 	}
 
-	SelectItems(ItemsToSelect
-		, ENavigationToolItemSelectionFlags::AppendToCurrentSelection | ENavigationToolItemSelectionFlags::SignalSelectionChange);
+	SelectItems(WeakItemsToSelect, ENavigationToolItemSelectionFlags::AppendToCurrentSelection
+		| ENavigationToolItemSelectionFlags::SignalSelectionChange);
 }
 
 bool FNavigationToolView::CanSelectChildren() const
@@ -1461,27 +1569,32 @@ bool FNavigationToolView::CanSelectChildren() const
 
 void FNavigationToolView::SelectParent()
 {
-	const TSet<FNavigationToolItemPtr> Items(GetSelectedItems());
+	const TSet<FNavigationToolViewModelWeakPtr> WeakItems(GetSelectedItems());
 
-	TSet<FNavigationToolItemPtr> ParentItemsToSelect;
-	ParentItemsToSelect.Reserve(Items.Num());
+	TSet<FNavigationToolViewModelWeakPtr> WeakParentItemsToSelect;
+	WeakParentItemsToSelect.Reserve(WeakItems.Num());
 
-	const FNavigationToolItemPtr RootItem = GetRootItem();
+	const FNavigationToolViewModelPtr RootItem = GetRootItem();
 	
 	//Add only Valid Parents that are not Root and are not part of the Original Selection!
-	for (const FNavigationToolItemPtr& Item : Items)
+	for (const FNavigationToolViewModelWeakPtr& WeakItem : WeakItems)
 	{
-		if (Item.IsValid())
+		const FNavigationToolViewModelPtr Item = WeakItem.Pin();
+		if (!Item.IsValid())
 		{
-			const FNavigationToolItemPtr ParentItem = Item->GetParent();
-			if (ParentItem.IsValid() && ParentItem != RootItem && !Items.Contains(ParentItem))
-			{
-				ParentItemsToSelect.Add(ParentItem);
-			}
+			continue;
+		}
+
+		const FNavigationToolViewModelPtr ParentItem = Item->GetParent();
+		if (ParentItem.IsValid()
+			&& ParentItem != RootItem
+			&& !WeakItems.Contains(ParentItem))
+		{
+			WeakParentItemsToSelect.Add(ParentItem);
 		}
 	}
 
-	SortAndSelectItems(ParentItemsToSelect.Array());
+	SortAndSelectItems(WeakParentItemsToSelect.Array());
 }
 
 bool FNavigationToolView::CanSelectParent() const
@@ -1491,26 +1604,29 @@ bool FNavigationToolView::CanSelectParent() const
 
 void FNavigationToolView::SelectFirstChild()
 {
-	const TArray<FNavigationToolItemPtr> Items = GetSelectedItems();
+	const TSet<FNavigationToolViewModelWeakPtr> WeakItems(GetSelectedItems());
 
-	TSet<FNavigationToolItemPtr> FirstChildItemsToSelect;
-	FirstChildItemsToSelect.Reserve(Items.Num());
+	TSet<FNavigationToolViewModelWeakPtr> WeakFirstChildItemsToSelect;
+	WeakFirstChildItemsToSelect.Reserve(WeakItems.Num());
 
-	for (const FNavigationToolItemPtr& Item : Items)
+	for (const FNavigationToolViewModelWeakPtr& WeakItem : WeakItems)
 	{
-		if (Item.IsValid())
+		const FNavigationToolViewModelPtr Item = WeakItem.Pin();
+		if (!Item.IsValid())
 		{
-			const FNavigationToolItemPtr FirstChildItem = GetVisibleChildAt(Item, 0);
+			continue;
+		}
 
-			//Don't select Component items! (Component items on selection also select their owner actor items, which can cause undesired issues)
-			if (FirstChildItem.IsValid() && !FirstChildItem->IsA<FNavigationToolComponent>())
-			{
-				FirstChildItemsToSelect.Add(FirstChildItem);
-			}
+		const FNavigationToolViewModelPtr FirstChildItem = GetVisibleChildAt(Item, 0);
+
+		// Don't select Component items! (Component items on selection also select their owner actor items, which can cause undesired issues)
+		if (FirstChildItem.IsValid() && !FirstChildItem.AsModel()->IsA<FNavigationToolComponent>())
+		{
+			WeakFirstChildItemsToSelect.Add(FirstChildItem);
 		}
 	}
-	
-	SortAndSelectItems(FirstChildItemsToSelect.Array());
+
+	SortAndSelectItems(WeakFirstChildItemsToSelect.Array());
 }
 
 bool FNavigationToolView::CanSelectFirstChild() const
@@ -1518,33 +1634,41 @@ bool FNavigationToolView::CanSelectFirstChild() const
 	return GetViewSelectedItemCount() == 1;
 }
 
-void FNavigationToolView::SelectSibling(int32 InDeltaIndex)
+void FNavigationToolView::SelectSibling(const int32 InDeltaIndex)
 {
-	const TArray<FNavigationToolItemPtr> Items = GetSelectedItems();
+	const TArray<FNavigationToolViewModelWeakPtr> WeakItems = GetSelectedItems();
 
-	TSet<FNavigationToolItemPtr> SiblingItemsToSelect;
-	SiblingItemsToSelect.Reserve(Items.Num());
+	TSet<FNavigationToolViewModelWeakPtr> WeakSiblingItemsToSelect;
+	WeakSiblingItemsToSelect.Reserve(WeakItems.Num());
 
-	for (const FNavigationToolItemPtr& Item : Items)
+	for (const FNavigationToolViewModelWeakPtr& WeakItem : WeakItems)
 	{
-		if (Item.IsValid() && Item->GetParent())
+		const FNavigationToolViewModelPtr Item = WeakItem.Pin();
+		if (!Item.IsValid())
 		{
-			const FNavigationToolItemPtr ParentItem = Item->GetParent();
+			continue;
+		}
 
-			const int32 ItemIndex   = GetVisibleChildIndex(ParentItem, Item);
-			const int32 TargetIndex = ItemIndex + InDeltaIndex;
+		const FNavigationToolViewModelPtr ParentItem = Item->GetParent();
+		if (!ParentItem.IsValid())
+		{
+			continue;
+		}
 
-			//Don't try to Normalize Index, if it's Invalid, we won't cycle and just skip that selection
-			const FNavigationToolItemPtr SiblingToSelect = GetVisibleChildAt(ParentItem, TargetIndex);
+		const int32 ItemIndex = GetVisibleChildIndex(ParentItem, Item);
+		const int32 TargetIndex = ItemIndex + InDeltaIndex;
 
-			//Don't select Component items! (Component items on selection also select their owner actor items, which can cause undesired issues)
-			if (SiblingToSelect.IsValid() && !SiblingToSelect->IsA<FNavigationToolComponent>())
-			{
-				SiblingItemsToSelect.Add(SiblingToSelect);
-			}
+		// Don't try to Normalize Index, if it's Invalid, we won't cycle and just skip that selection
+		const FNavigationToolViewModelPtr SiblingToSelect = GetVisibleChildAt(ParentItem, TargetIndex);
+
+		// Don't select Component items! (Component items on selection also select their owner actor items, which can cause undesired issues)
+		if (SiblingToSelect.IsValid() && !SiblingToSelect.AsModel()->IsA<FNavigationToolComponent>())
+		{
+			WeakSiblingItemsToSelect.Add(SiblingToSelect);
 		}
 	}
-	SortAndSelectItems(SiblingItemsToSelect.Array());
+
+	SortAndSelectItems(WeakSiblingItemsToSelect.Array());
 }
 
 bool FNavigationToolView::CanSelectSibling() const
@@ -1554,9 +1678,9 @@ bool FNavigationToolView::CanSelectSibling() const
 
 void FNavigationToolView::ExpandAll()
 {
-	for (const FNavigationToolItemPtr& Item : RootVisibleItems)
+	for (const FNavigationToolViewModelWeakPtr& WeakItem : WeakRootVisibleItems)
 	{
-		SetItemExpansionRecursive(Item, true);
+		SetItemExpansionRecursive(WeakItem.Pin(), true);
 	}
 }
 
@@ -1567,9 +1691,9 @@ bool FNavigationToolView::CanExpandAll() const
 
 void FNavigationToolView::CollapseAll()
 {
-	for (const FNavigationToolItemPtr& Item : RootVisibleItems)
+	for (const FNavigationToolViewModelWeakPtr& WeakItem : WeakRootVisibleItems)
 	{
-		SetItemExpansionRecursive(Item, false);
+		SetItemExpansionRecursive(WeakItem.Pin(), false);
 	}
 }
 
@@ -1593,47 +1717,53 @@ bool FNavigationToolView::CanScrollNextIntoView() const
 	return GetViewSelectedItemCount() > 0;
 }
 
-void FNavigationToolView::ScrollDeltaIndexIntoView(int32 DeltaIndex)
+void FNavigationToolView::ScrollDeltaIndexIntoView(const int32 InDeltaIndex)
 {
-	const int32 SelectedItemCount = SortedSelectedItems.Num();
-	if (SortedSelectedItems.Num() > 0)
+	const int32 SelectedItemCount = WeakSortedSelectedItems.Num();
+	if (WeakSortedSelectedItems.Num() > 0)
 	{
-		const int32 TargetIndex  = NextSelectedItemIntoView + DeltaIndex;
+		const int32 TargetIndex  = NextSelectedItemIntoView + InDeltaIndex;
 		NextSelectedItemIntoView = TargetIndex % SelectedItemCount;
 		if (NextSelectedItemIntoView < 0)
 		{
 			NextSelectedItemIntoView += SelectedItemCount;
 		}
-		ScrollItemIntoView(SortedSelectedItems[NextSelectedItemIntoView]);
+		ScrollItemIntoView(WeakSortedSelectedItems[NextSelectedItemIntoView].Pin());
 	}
 }
 
-void FNavigationToolView::ScrollItemIntoView(const FNavigationToolItemPtr& InItem)
+void FNavigationToolView::ScrollItemIntoView(const FNavigationToolViewModelWeakPtr& InWeakItem)
 {
-	if (InItem.IsValid())
+	const FNavigationToolViewModelPtr Item = InWeakItem.Pin();
+	if (!Item.IsValid())
 	{
-		SetParentItemExpansions(InItem, true);
-		if (ToolViewWidget.IsValid() && ToolViewWidget->GetTreeView())
-		{
-			ToolViewWidget->GetTreeView()->FocusOnItem(InItem);
-			ToolViewWidget->ScrollItemIntoView(InItem);
-		}
+		return;
+	}
+
+	SetParentItemExpansions(InWeakItem, true);
+
+	if (ToolViewWidget.IsValid() && ToolViewWidget->GetTreeView())
+	{
+		ToolViewWidget->GetTreeView()->FocusOnItem(Item);
+		ToolViewWidget->ScrollItemIntoView(InWeakItem);
 	}
 }
 
-void FNavigationToolView::SortAndSelectItems(TArray<FNavigationToolItemPtr> InItemsToSelect)
+void FNavigationToolView::SortAndSelectItems(TArray<FNavigationToolViewModelWeakPtr> InWeakItemsToSelect)
 {
-	if (!InItemsToSelect.IsEmpty())
+	if (InWeakItemsToSelect.IsEmpty())
 	{
-		FNavigationTool::SortItems(InItemsToSelect);
-
-		SelectItems(InItemsToSelect
-			, ENavigationToolItemSelectionFlags::SignalSelectionChange
-			| ENavigationToolItemSelectionFlags::ScrollIntoView);
+		return;
 	}
+
+	FNavigationTool::SortItems(InWeakItemsToSelect);
+
+	SelectItems(InWeakItemsToSelect
+		, ENavigationToolItemSelectionFlags::SignalSelectionChange
+		| ENavigationToolItemSelectionFlags::ScrollIntoView);
 }
 
-void FNavigationToolView::RefreshTool(bool bInImmediateRefresh)
+void FNavigationToolView::RefreshTool(const bool bInImmediateRefresh)
 {
 	if (const TSharedPtr<INavigationTool> Tool = GetOwnerTool())
 	{
@@ -1812,61 +1942,86 @@ void FNavigationToolView::LoadFilterState(const FNavigationToolViewSaveState& In
 
 void FNavigationToolView::SaveToolViewItems(FNavigationToolViewSaveState& OutViewSaveState)
 {
-	const TSharedRef<FNavigationToolItem> TreeRoot = GetOwnerTool()->GetTreeRoot();
+	const TSharedPtr<FNavigationTool> OwnerTool = WeakTool.Pin();
+	if (!OwnerTool.IsValid())
+	{
+		return;
+	}
 
-	TArray<FNavigationToolItemPtr> ItemsToSave = TreeRoot->GetChildren();
+	const FNavigationToolViewModelPtr TreeRoot = OwnerTool->GetTreeRoot().Pin();
+	if (!TreeRoot.IsValid())
+	{
+		return;
+	}
+
+	TArray<FNavigationToolViewModelWeakPtr> WeakItemsToSave = TreeRoot->GetChildren();
 
 	OutViewSaveState.ViewItemFlags.Reset();
 
-	while (ItemsToSave.Num() > 0)
+	while (WeakItemsToSave.Num() > 0)
 	{
-		FNavigationToolItemPtr ItemToSave = ItemsToSave.Pop();
-		if (ItemToSave.IsValid())
+		const FNavigationToolViewModelPtr ItemToSave = WeakItemsToSave.Pop().Pin();
+		if (!ItemToSave.IsValid())
 		{
-			// Iteratively also Save children
-			ItemsToSave.Append(ItemToSave->GetChildren());
+			continue;
+		}
 
-			const FNavigationToolItemId ItemId = ItemToSave->GetItemId();
-			const FString StringId = ItemId.GetStringId();
+		// Iteratively also save children
+		WeakItemsToSave.Append(ItemToSave->GetChildren());
 
-			// Save Item State Flags
-			if (const ENavigationToolItemFlags* const ItemFlags = OutViewSaveState.ViewItemFlags.Find(StringId))
-			{
-				OutViewSaveState.ViewItemFlags.Add(StringId, *ItemFlags);
-			}
-			else
-			{
-				OutViewSaveState.ViewItemFlags.Remove(StringId);
-			}
+		const FNavigationToolItemId ItemId = ItemToSave->GetItemId();
+		const FString StringId = ItemId.GetStringId();
+
+		// Save item state flags
+		if (const ENavigationToolItemFlags* const ItemFlags = OutViewSaveState.ViewItemFlags.Find(StringId))
+		{
+			OutViewSaveState.ViewItemFlags.Add(StringId, *ItemFlags);
+		}
+		else
+		{
+			OutViewSaveState.ViewItemFlags.Remove(StringId);
 		}
 	}
 }
 
 void FNavigationToolView::LoadToolViewItems(FNavigationToolViewSaveState& InViewSaveState)
 {
-	const TSharedRef<FNavigationToolItem> TreeRoot = GetOwnerTool()->GetTreeRoot();
+	const TSharedPtr<FNavigationTool> OwnerTool = WeakTool.Pin();
+	if (!OwnerTool.IsValid())
+	{
+		return;
+	}
+
+	const FNavigationToolViewModelPtr TreeRoot = OwnerTool->GetTreeRoot().Pin();
+	if (!TreeRoot.IsValid())
+	{
+		return;
+	}
 
-	TArray<FNavigationToolItemPtr> ItemsToLoad = TreeRoot->GetChildren();
+	TArray<FNavigationToolViewModelWeakPtr> WeakItemsToLoad = TreeRoot->GetChildren();
 
-	while (ItemsToLoad.Num() > 0)
+	while (WeakItemsToLoad.Num() > 0)
 	{
-		if (const FNavigationToolItemPtr ItemToLoad = ItemsToLoad.Pop())
+		const FNavigationToolViewModelPtr ItemToLoad = WeakItemsToLoad.Pop().Pin();
+		if (!ItemToLoad.IsValid())
 		{
-			// Iteratively also Load Children
-			ItemsToLoad.Append(ItemToLoad->GetChildren());
+			continue;
+		}
 
-			const FNavigationToolItemId ItemId = ItemToLoad->GetItemId();
-			const FString StringId = ItemId.GetStringId();
+		// Iteratively also load children
+		WeakItemsToLoad.Append(ItemToLoad->GetChildren());
 
-			// Load Item Flags
-			if (const ENavigationToolItemFlags* const ItemFlags = InViewSaveState.ViewItemFlags.Find(StringId))
-			{
-				InViewSaveState.ViewItemFlags.Add(StringId, *ItemFlags);
-			}
-			else
-			{
-				InViewSaveState.ViewItemFlags.Remove(StringId);
-			}
+		const FNavigationToolItemId ItemId = ItemToLoad->GetItemId();
+		const FString StringId = ItemId.GetStringId();
+
+		// Load item flags
+		if (const ENavigationToolItemFlags* const ItemFlags = InViewSaveState.ViewItemFlags.Find(StringId))
+		{
+			InViewSaveState.ViewItemFlags.Add(StringId, *ItemFlags);
+		}
+		else
+		{
+			InViewSaveState.ViewItemFlags.Remove(StringId);
 		}
 	}
 }
@@ -2108,23 +2263,32 @@ void FNavigationToolView::ApplyCustomColumnView(const FText InColumnViewName)
 
 bool FNavigationToolView::CanFocusSingleSelection() const
 {
-	if (SelectedItems.Num() == 1)
+	if (WeakSelectedItems.Num() == 1)
 	{
-		const FNavigationToolSequence* const SequenceItem = SelectedItems[0]->CastTo<FNavigationToolSequence>();
-		return SequenceItem && SequenceItem->GetSequence();
+		if (const FNavigationToolViewModelPtr SelectedItem = WeakSelectedItems[0].Pin())
+		{
+			const TViewModelPtr<FNavigationToolSequence> SequenceItem = SelectedItem.ImplicitCast();
+			return SequenceItem && SequenceItem->GetSequence();
+		}
 	}
 	return false;
 }
 
 void FNavigationToolView::FocusSingleSelection()
 {
-	if (SelectedItems.Num() != 1)
+	if (WeakSelectedItems.Num() != 1)
 	{
 		return;
 	}
 
-	const FNavigationToolSequence* const SequenceItem = SelectedItems[0]->CastTo<FNavigationToolSequence>();
-	if (!SequenceItem)
+	const FNavigationToolViewModelPtr SelectedItem = WeakSelectedItems[0].Pin();
+	if (!SelectedItem.IsValid())
+	{
+		return;
+	}
+
+	const TViewModelPtr<FNavigationToolSequence> SequenceItem = SelectedItem.ImplicitCast();
+	if (!SequenceItem.IsValid())
 	{
 		return;
 	}
@@ -2146,13 +2310,19 @@ void FNavigationToolView::FocusSingleSelection()
 
 bool FNavigationToolView::CanFocusInContentBrowser() const
 {
-	if (SelectedItems.Num() != 1)
+	if (WeakSelectedItems.Num() != 1)
 	{
 		return false;
 	}
 
-	const FNavigationToolSequence* const SequenceItem = SelectedItems[0]->CastTo<FNavigationToolSequence>();
-	if (!SequenceItem)
+	const FNavigationToolViewModelPtr SelectedItem = WeakSelectedItems[0].Pin();
+	if (!SelectedItem.IsValid())
+	{
+		return false;
+	}
+
+	const TViewModelPtr<FNavigationToolSequence> SequenceItem = SelectedItem.ImplicitCast();
+	if (!SequenceItem.IsValid())
 	{
 		return false;
 	}
@@ -2168,13 +2338,19 @@ bool FNavigationToolView::CanFocusInContentBrowser() const
 
 void FNavigationToolView::FocusInContentBrowser()
 {
-	if (SelectedItems.Num() != 1)
+	if (WeakSelectedItems.Num() != 1)
+	{
+		return;
+	}
+
+	const FNavigationToolViewModelPtr SelectedItem = WeakSelectedItems[0].Pin();
+	if (!SelectedItem.IsValid())
 	{
 		return;
 	}
 
-	const FNavigationToolSequence* const SequenceItem = SelectedItems[0]->CastTo<FNavigationToolSequence>();
-	if (!SequenceItem)
+	const TViewModelPtr<FNavigationToolSequence> SequenceItem = SelectedItem.ImplicitCast();
+	if (!SequenceItem.IsValid())
 	{
 		return;
 	}
diff --git a/Engine/Plugins/Experimental/SequenceNavigator/Source/SequenceNavigator/Private/NavigationToolView.h b/Engine/Plugins/Experimental/SequenceNavigator/Source/SequenceNavigator/Private/NavigationToolView.h
--- a/Engine/Plugins/Experimental/SequenceNavigator/Source/SequenceNavigator/Private/NavigationToolView.h
+++ b/Engine/Plugins/Experimental/SequenceNavigator/Source/SequenceNavigator/Private/NavigationToolView.h
@@ -6,7 +6,6 @@
 #include "Filters/NavigationToolFilterBar.h"
 #include "INavigationToolView.h"
 #include "NavigationToolDefines.h"
-#include "NavigationToolItemType.h"
 
 class FDragDropEvent;
 class FReply;
@@ -63,9 +62,9 @@ public:
 	void PostLoad();
 
 	/** Called when the UNavigationToolSettings has a property change */
-	void OnToolSettingsChanged(UObject* InObject, FPropertyChangedEvent& InPropertyChangedEvent);
+	void OnToolSettingsChanged(UObject* const InObject, FPropertyChangedEvent& InPropertyChangedEvent);
 
-	void Tick(float InDeltaTime);
+	void Tick(const float InDeltaTime);
 
 	void BindCommands(const TSharedPtr<FUICommandList>& InBaseCommandList);
 
@@ -90,63 +89,81 @@ public:
 	virtual ENavigationToolItemViewMode GetItemDefaultViewMode() const override;
 	virtual ENavigationToolItemViewMode GetItemProxyViewMode() const override;
 	virtual TSharedRef<SWidget> GetColumnMenuContent(const FName InColumnId) override;
-	virtual void GetChildrenOfItem(const FNavigationToolItemPtr InItem, TArray<FNavigationToolItemPtr>& OutChildren) const override;
-	virtual void GetChildrenOfItem(const FNavigationToolItemPtr& InItem
-		, TArray<FNavigationToolItemPtr>& OutChildren
-		, ENavigationToolItemViewMode InViewMode
-		, const TSet<FNavigationToolItemPtr>& InRecursionDisallowedItems) const override;
-	virtual bool IsItemReadOnly(const FNavigationToolItemPtr& InItem) const override;
-	virtual bool CanSelectItem(const FNavigationToolItemPtr& InItem) const override;
-	virtual void SelectItems(TArray<FNavigationToolItemPtr> InItems, const ENavigationToolItemSelectionFlags InFlags = ENavigationToolItemSelectionFlags::SignalSelectionChange) override;
-	virtual void ClearItemSelection(bool bSignalSelectionChange = true) override;
-	virtual bool IsItemSelected(const FNavigationToolItemPtr& InItem) const override;
-	virtual TArray<FNavigationToolItemPtr> GetSelectedItems() const override;
-	virtual FReply OnDragDetected(const FGeometry& InGeometry, const FPointerEvent& InMouseEvent, const FNavigationToolItemPtr InTargetItem) override;
-	virtual FReply OnDrop(const FDragDropEvent& InDragDropEvent, EItemDropZone InDropZone, const FNavigationToolItemPtr InTargetItem) override;
+	virtual void GetChildrenOfItem(const FNavigationToolViewModelWeakPtr InWeakItem
+		, TArray<FNavigationToolViewModelWeakPtr>& OutWeakChildren) const override;
+	virtual void GetChildrenOfItem(const FNavigationToolViewModelWeakPtr InWeakItem
+		, TArray<FNavigationToolViewModelWeakPtr>& OutWeakChildren
+		, const ENavigationToolItemViewMode InViewMode
+		, const TSet<FNavigationToolViewModelWeakPtr>& InWeakRecursionDisallowedItems) const override;
+	virtual bool IsItemReadOnly(const FNavigationToolViewModelWeakPtr& InWeakItem) const override;
+	virtual bool CanSelectItem(const FNavigationToolViewModelWeakPtr& InWeakItem) const override;
+	virtual void SelectItems(TArray<FNavigationToolViewModelWeakPtr> InWeakItems
+		, const ENavigationToolItemSelectionFlags InFlags = ENavigationToolItemSelectionFlags::SignalSelectionChange) override;
+	virtual void ClearItemSelection(const bool bInSignalSelectionChange = true) override;
+	virtual bool IsItemSelected(const FNavigationToolViewModelWeakPtr& InWeakItem) const override;
+	virtual TArray<FNavigationToolViewModelWeakPtr> GetSelectedItems() const override;
+	/*virtual FReply OnDragDetected(const FGeometry& InGeometry
+		, const FPointerEvent& InPointerEvent
+		, const FNavigationToolViewModelWeakPtr InWeakTargetItem) override;
+	virtual FReply OnDrop(const FDragDropEvent& InDragDropEvent
+		, const EItemDropZone InDropZone
+		, const FNavigationToolViewModelWeakPtr InWeakTargetItem) override;
 	virtual TOptional<EItemDropZone> OnCanDrop(const FDragDropEvent& InDragDropEvent
-		, EItemDropZone InDropZone
-		, const FNavigationToolItemPtr InTargetItem) const override;
-	virtual bool IsItemExpanded(const FNavigationToolItemPtr& InItem, const bool bInUseFilter = true) const override;
-	virtual void SetItemExpansion(const FNavigationToolItemPtr& InItem, const bool bInExpand, const bool bInUseFilter = true) override;
-	virtual void SetItemExpansionRecursive(const FNavigationToolItemPtr InItem, const bool bInExpand) override;
-	virtual void SetParentItemExpansions(const FNavigationToolItemPtr& InItem, const bool bInExpand) override;
+		, const EItemDropZone InDropZone
+		, const FNavigationToolViewModelWeakPtr InWeakTargetItem) const override;*/
+	virtual bool IsItemExpanded(const FNavigationToolViewModelWeakPtr& InWeakItem, const bool bInUseFilter = true) const override;
+	virtual void SetItemExpansion(const FNavigationToolViewModelWeakPtr& InWeakItem
+		, const bool bInExpand, const bool bInUseFilter = true) override;
+	virtual void SetItemExpansionRecursive(const FNavigationToolViewModelWeakPtr InWeakItem, const bool bInExpand) override;
+	virtual void SetParentItemExpansions(const FNavigationToolViewModelWeakPtr& InWeakItem, const bool bInExpand) override;
 	virtual bool CanExpandAll() const override;
 	virtual void ExpandAll() override;
 	virtual void CollapseAll() override;
 	virtual bool CanCollapseAll() const override;
 	//~ End INavigationToolView
 
+	FReply OnDragDetected(const FGeometry& InGeometry
+		, const FPointerEvent& InPointerEvent
+		, const FNavigationToolViewModelWeakPtr InWeakTargetItem);
+	FReply OnDrop(const FDragDropEvent& InDragDropEvent
+		, const EItemDropZone InDropZone
+		, const FNavigationToolViewModelWeakPtr InWeakTargetItem);
+	TOptional<EItemDropZone> OnCanDrop(const FDragDropEvent& InDragDropEvent
+		, const EItemDropZone InDropZone
+		, const FNavigationToolViewModelWeakPtr InWeakTargetItem) const;
+
 	TSharedPtr<SWidget> CreateItemContextMenu();
 
-	/** Get the Columns currently allocated in this Navigation Tool Instance */
+	/** Get the columns currently allocated in this Navigation Tool instance */
 	const TMap<FName, TSharedPtr<INavigationToolColumn>>& GetColumns() const { return Columns; }
 
-	/** Whether the given column should be shown by default. This is only used at the start when Creating the Columns */
+	/** Whether the given column should be shown by default. This is only used at the start when creating the columns */
 	bool ShouldShowColumnByDefault(const TSharedPtr<INavigationToolColumn>& InColumn) const;
 
-	/** Refreshes the Items visible in this View, and refreshes the Navigation Tool Widget if it was created */
+	/** Refreshes the items visible in this view and refreshes the Navigation Tool widget if it was created */
 	void Refresh();
 
-	/** Refreshes the Items that will be at the Top Level of the Tree (e.g. Actors with no parent item) */
+	/** Refreshes the items that will be at the top level of the tree (e.g. Actors with no parent item) */
 	void UpdateRootVisibleItems();
 
-	/** Updates the Item Expansions in the Tree Widget based on whether the Items have the Expanded flag or not */
+	/** Updates the item expansions in the tree widget based on whether the Items have the expanded flag or not */
 	void UpdateItemExpansions();
 
-	/** Called when Object Replacement has taken place. Used to invalidate the widget for painting */
+	/** Called when object replacement has taken place. Used to invalidate the widget for painting */
 	void NotifyObjectsReplaced();
 
-	/** The Root Item of the Tree. Convenience function to get the RootItem from FNavigationTool */
-	FNavigationToolItemPtr GetRootItem() const;
+	/** The root item of the tree. Convenience function to get the RootItem from FNavigationTool */
+	FNavigationToolViewModelPtr GetRootItem() const;
 
-	/** Gets the Top Level Items that should be visible in the Tree. This can differ between views depending on Filters and other factors */
-	const TArray<FNavigationToolItemPtr>& GetRootVisibleItems() const;
+	/** Gets the top level items that should be visible in the tree. This can differ between views depending on filters and other factors */
+	const TArray<FNavigationToolViewModelWeakPtr>& GetRootVisibleItems() const;
 
-	void SaveViewItemFlags(const FNavigationToolItemPtr& InItem, ENavigationToolItemFlags InFlags);
+	void SaveViewItemFlags(const FNavigationToolViewModelWeakPtr& InWeakItem
+		, const ENavigationToolItemFlags InFlags);
 	
-	ENavigationToolItemFlags GetViewItemFlags(const FNavigationToolItemPtr& InItem) const;
+	ENavigationToolItemFlags GetViewItemFlags(const FNavigationToolViewModelWeakPtr& InWeakItem) const;
 
-	FLinearColor GetItemBrushColor(const FNavigationToolItemPtr InItem) const;
+	FLinearColor GetItemBrushColor(const FNavigationToolViewModelPtr InItem) const;
 
 	/** Gets the Currently Selected Item Count in the Tree View */
 	int32 GetViewSelectedItemCount() const;
@@ -160,18 +177,19 @@ public:
 	TSharedPtr<FNavigationToolFilterBar> GetFilterBar() const { return FilterBar; }
 
 	/** Called when Item selection has changed. This can be called by the Navigation Tool Widget or by the View itself if there's no ToolWidget */
-	void NotifyItemSelectionChanged(const TArray<FNavigationToolItemPtr>& InSelectedItems
-		, const FNavigationToolItemPtr& InItem
-		, bool bUpdateModeTools);
+	void NotifyItemSelectionChanged(const TArray<FNavigationToolViewModelWeakPtr>& InWeakSelectedItems
+		, const FNavigationToolViewModelPtr& InItem
+		, const bool bInUpdateModeTools);
 
 	/** Return whether the given item should be visible in this Navigation Tool Instance */
-	bool ShouldShowItem(const FNavigationToolItemPtr& InItem, bool bInUseFilters, ENavigationToolItemViewMode InViewMode) const;
+	bool ShouldShowItem(const FNavigationToolViewModelWeakPtr& InWeakItem
+		, const bool bInUseFilters, const ENavigationToolItemViewMode InViewMode) const;
 
 	/** Returns the Index of the Child from the Parent's Children Visible List (Visible List can vary between Tool Instances) */
-	int32 GetVisibleChildIndex(const FNavigationToolItemPtr& InParentItem, const FNavigationToolItemPtr& InChildItem) const;
+	int32 GetVisibleChildIndex(const FNavigationToolViewModelPtr& InParentItem, const FNavigationToolViewModelPtr& InChildItem) const;
 
 	/** Returns the Child from the Parent's Children Visible List at the given Index (Visible List can vary between Tool Instances) */
-	FNavigationToolItemPtr GetVisibleChildAt(const FNavigationToolItemPtr& InParentItem, int32 InChildIndex) const;
+	FNavigationToolViewModelPtr GetVisibleChildAt(const FNavigationToolViewModelPtr& InParentItem, const int32 InChildIndex) const;
 
 	/** Returns whether the Tool is Locked. See FNavigationTool::IsToolLocked */
 	bool IsToolLocked() const;
@@ -244,19 +262,19 @@ public:
 	template<typename InItemType, typename = typename TEnableIf<TIsDerivedFrom<InItemType, INavigationToolItem>::IsDerived>::Type>
 	void HideItemType()
 	{
-		SetItemTypeHidden(TNavigationToolItemType<InItemType>::GetTypeId().ToName(), true);
+		SetItemTypeHidden(InItemType::GetTypeTable().GetTypeName(), true);
 	}
 	
 	/** Shows the given Outliner Item Type if it was hidden in this Outliner View (e.g. Showing Components again) */
 	template<typename InItemType, typename = typename TEnableIf<TIsDerivedFrom<InItemType, INavigationToolItem>::IsDerived>::Type>
 	void ExposeItemType()
 	{
-		SetItemTypeHidden(TNavigationToolItemType<InItemType>::GetTypeId().ToName(), false);
+		SetItemTypeHidden(InItemType::GetTypeTable().GetTypeName(), false);
 	}
 	
 	/** Whether the given Outliner Item Type is currently hidden in this Outliner View */
 	bool IsItemTypeHidden(const FName InItemTypeName) const;
-	bool IsItemTypeHidden(const FNavigationToolItemPtr& InItem) const;
+	bool IsItemTypeHidden(const FNavigationToolViewModelPtr& InItem) const;
 	
 	template<typename InItemType, typename = typename TEnableIf<TIsDerivedFrom<InItemType, INavigationToolItem>::IsDerived>::Type>
 	bool IsItemTypeHidden() const
@@ -264,15 +282,17 @@ public:
 		return IsItemTypeHidden(InItemType::GetStaticTypeName());
 	}
 
-	/** Called when a Drag enters the Navigation Tool Widgets for the given Item (if null, treat as root) */
-	void OnDragEnter(const FDragDropEvent& InDragDropEvent, const FNavigationToolItemPtr InTargetItem);
+	/** Called when a drag enters the Navigation Tool widgets for the given item (if null, treat as root) */
+	void OnDragEnter(const FDragDropEvent& InDragDropEvent
+		, const FNavigationToolViewModelWeakPtr InWeakTargetItem);
 
-	/** Called when a Drag leaves the Navigation Tool Widgets for the given Item (if null, treat as root) */
-	void OnDragLeave(const FDragDropEvent& InDragDropEvent, const FNavigationToolItemPtr InTargetItem);
+	/** Called when a drag leaves the Navigation Tool widgets for the given item (if null, treat as root) */
+	void OnDragLeave(const FDragDropEvent& InDragDropEvent
+		, const FNavigationToolViewModelWeakPtr InWeakTargetItem);
+
+	/** Called when there's a drag being started/ended to the Navigation Tool widget as a whole (i.e. root) and not an individual item */
+	void SetDragIntoTreeRoot(const bool bInIsDraggingIntoTreeRoot);
 
-	/** Called when there's a Drag being started/ended to the Navigation Tool Widget as a whole (i.e. root) and not an individual item */
-	void SetDragIntoTreeRoot(bool bInIsDraggingIntoTreeRoot);
-	
 	void RenameSelected();
 	void ResetRenaming();
 	void OnItemRenameAction(const ENavigationToolRenameAction InRenameAction, const TSharedPtr<INavigationToolView>& InToolView);
@@ -284,7 +304,7 @@ public:
 	void DuplicateSelected();
 	bool CanDuplicateSelected() const;
 
-	void SelectChildren(bool bIsRecursive);
+	void SelectChildren(const bool bInIsRecursive);
 	bool CanSelectChildren() const;
 
 	void SelectParent();
@@ -293,20 +313,20 @@ public:
 	void SelectFirstChild();
 	bool CanSelectFirstChild() const;
 
-	void SelectSibling(int32 InDeltaIndex);
+	void SelectSibling(const int32 InDeltaIndex);
 	bool CanSelectSibling() const;
 
 	/** Called when the given item's expansion state (expanded/collapsed) has changed */
-	void OnItemExpansionChanged(const FNavigationToolItemPtr InItem, const bool bInIsExpanded);
+	void OnItemExpansionChanged(const FNavigationToolViewModelWeakPtr InWeakItem, const bool bInIsExpanded);
 
 	void ScrollNextIntoView();
 	void ScrollPrevIntoView();
 	bool CanScrollNextIntoView() const;
-	void ScrollDeltaIndexIntoView(int32 DeltaIndex);
-	void ScrollItemIntoView(const FNavigationToolItemPtr& InItem);
+	void ScrollDeltaIndexIntoView(const int32 InDeltaIndex);
+	void ScrollItemIntoView(const FNavigationToolViewModelWeakPtr& InWeakItem);
 
 	/** Sorts the given Items in the user-defined order and selects them, and if the Widget is valid focusing to the item in the Widget */
-	void SortAndSelectItems(TArray<FNavigationToolItemPtr> InItemsToSelect);
+	void SortAndSelectItems(TArray<FNavigationToolViewModelWeakPtr> InWeakItemsToSelect);
 
 	/** Save the current state of the Navigation Tool */
 	void SaveViewState(const TSharedRef<FNavigationToolProvider>& InProvider);
@@ -333,11 +353,10 @@ protected:
 	void Init(const TSharedRef<FNavigationTool>& InTool, const TSharedPtr<FUICommandList>& InBaseCommandList);
 
 	/** Sets the Selected Items to be the given Items in this List. */
-	void SetItemSelectionImpl(TArray<FNavigationToolItemPtr>&& InItems
-		, bool bSignalSelectionChange);
+	void SetItemSelectionImpl(TArray<FNavigationToolViewModelWeakPtr>&& InWeakItems, const bool bSignalSelectionChange);
 
 	/** Triggers a Refresh on the FNavigationTool */
-	void RefreshTool(bool bInImmediateRefresh);
+	void RefreshTool(const bool bInImmediateRefresh);
 
 	void EnsureToolViewCount(const int32 InToolViewId) const;
 
@@ -377,13 +396,13 @@ protected:
 	const TObjectPtr<UNavigationToolSettings> ToolSettings;
 
 	/** Root Items from Navigation Tool visible to this Instance. i.e. a Subset of all the Items in FNavigationTool */
-	TArray<FNavigationToolItemPtr> RootVisibleItems;
+	TArray<FNavigationToolViewModelWeakPtr> WeakRootVisibleItems;
 
 	/** A list of the Selected Items in this Navigation Tool View */
-	TArray<FNavigationToolItemPtr> SelectedItems;
+	TArray<FNavigationToolViewModelWeakPtr> WeakSelectedItems;
 
 	/** Set of Items that are Currently Read Only in this Instance. Mutable as this changes on GetChildrenOfItem const func */
-	mutable TSet<FNavigationToolItemPtr> ReadOnlyItems;
+	mutable TSet<FNavigationToolViewModelWeakPtr> WeakReadOnlyItems;
 
 	/** A set of the Navigation Tool Item Type Names that should be hidden */
 	TSet<FName> HiddenItemTypes;
@@ -403,13 +422,13 @@ protected:
 	int32 NextSelectedItemIntoView = -1;
 
 	/** Selected Items Sorted from Top To Bottom */
-	TArray<FNavigationToolItemPtr> SortedSelectedItems;
+	TArray<FNavigationToolViewModelWeakPtr> WeakSortedSelectedItems;
 
 	/** The list of items that need renaming and are waiting for their turn to be renamed */
-	TArray<FNavigationToolItemPtr> ItemsRemainingRename;
+	TArray<FNavigationToolViewModelWeakPtr> WeakItemsRemainingRename;
 	
 	/** The current item in the process of renaming. Null if no renaming is taking place */
-	FNavigationToolItemPtr CurrentItemRenaming;
+	FNavigationToolViewModelWeakPtr WeakCurrentItemRenaming;
 	
 	/** Flag to call FNavigationToolView::Refresh next tick */
 	bool bRefreshRequested = false;
diff --git a/Engine/Plugins/Experimental/SequenceNavigator/Source/SequenceNavigator/Private/Providers/NavigationToolProvider.cpp b/Engine/Plugins/Experimental/SequenceNavigator/Source/SequenceNavigator/Private/Providers/NavigationToolProvider.cpp
--- a/Engine/Plugins/Experimental/SequenceNavigator/Source/SequenceNavigator/Private/Providers/NavigationToolProvider.cpp
+++ b/Engine/Plugins/Experimental/SequenceNavigator/Source/SequenceNavigator/Private/Providers/NavigationToolProvider.cpp
@@ -11,6 +11,8 @@
 namespace UE::SequenceNavigator
 {
 
+using namespace Sequencer;
+
 void FNavigationToolProvider::OnExtendColumnViews(TSet<FNavigationToolColumnView>& OutColumnViews)
 {
 	INavigationToolProvider::OnExtendColumnViews(OutColumnViews);
@@ -125,14 +127,14 @@ void FNavigationToolProvider::EnsureToolViewCount(const INavigationTool& InTool,
 
 TOptional<EItemDropZone> FNavigationToolProvider::OnToolItemCanAcceptDrop(const FDragDropEvent& InDragDropEvent
 	, EItemDropZone InDropZone
-	, const FNavigationToolItemPtr& InTargetItem) const
+	, const FNavigationToolViewModelPtr& InTargetItem) const
 {
 	return TOptional<EItemDropZone>();
 }
 
 FReply FNavigationToolProvider::OnToolItemAcceptDrop(const FDragDropEvent& InDragDropEvent
 	, EItemDropZone InDropZone
-	, const FNavigationToolItemPtr& InTargetItem)
+	, const FNavigationToolViewModelPtr& InTargetItem)
 {
 	return FReply::Unhandled();
 }
@@ -230,8 +232,17 @@ void FNavigationToolProvider::SaveState(FNavigationTool& InTool)
 
 void FNavigationToolProvider::LoadState(FNavigationTool& InTool)
 {
-	const FName ProviderId = GetIdentifier();
-	UE_LOG(LogSequenceNavigator, Log, TEXT("Loading provider state: %s"), *ProviderId.ToString());
+	const FNavigationToolViewModelPtr RootItem = InTool.GetTreeRoot().Pin();
+	if (!RootItem.IsValid())
+	{
+		return;
+	}
+
+	FNavigationToolSaveState* const SaveState = GetSaveState(InTool);
+	if (!SaveState)
+	{
+		return;
+	}
 
 	const UNavigationToolSettings* const ToolSettings = GetDefault<UNavigationToolSettings>();
 	if (!ToolSettings)
@@ -239,12 +250,12 @@ void FNavigationToolProvider::LoadState(FNavigationTool& InTool)
 		return;
 	}
 
+	const FName ProviderId = GetIdentifier();
+	UE_LOG(LogSequenceNavigator, Log, TEXT("Loading provider state: %s"), *ProviderId.ToString());
+
 	UpdateItemIdContexts(InTool);
 
-	if (FNavigationToolSaveState* const SaveState = GetSaveState(InTool))
-	{
-		LoadSerializedTree(InTool.GetTreeRoot(), &SaveState->SerializedTree);
-	}
+	LoadSerializedTree(RootItem, &SaveState->SerializedTree);
 
 	const TSharedRef<FNavigationToolProvider> SharedThisRef = SharedThis(this);
 
@@ -256,6 +267,12 @@ void FNavigationToolProvider::LoadState(FNavigationTool& InTool)
 
 void FNavigationToolProvider::SaveSerializedTree(FNavigationTool& InTool, const bool bInResetTree)
 {
+	const FNavigationToolViewModelPtr RootItem = InTool.GetTreeRoot().Pin();
+	if (!RootItem.IsValid())
+	{
+		return;
+	}
+
 	FNavigationToolSaveState* const SaveState = GetSaveState(InTool);
 	if (!SaveState)
 	{
@@ -267,19 +284,19 @@ void FNavigationToolProvider::SaveSerializedTree(FNavigationTool& InTool, const
 		SaveState->SerializedTree.Reset();
 	}
 
-	SaveSerializedTreeRecursive(InTool.GetTreeRoot(), SaveState->SerializedTree);
+	SaveSerializedTreeRecursive(RootItem, SaveState->SerializedTree);
 }
 
-void FNavigationToolProvider::SaveSerializedTreeRecursive(const FNavigationToolItemPtr& InParentItem
+void FNavigationToolProvider::SaveSerializedTreeRecursive(const FNavigationToolViewModelPtr& InParentItem
 	, FNavigationToolSerializedTree& InSerializedTree)
 {
-	const TArray<FNavigationToolItemPtr>& Children = InParentItem->GetChildren();
+	const TArray<FNavigationToolViewModelWeakPtr>& WeakChildren = InParentItem->GetChildren();
 
-	const FNavigationToolSerializedItem ParentSceneItem = InParentItem->MakeSerializedItem();//MakeSerializedItemFromToolItem(InParentItem);
+	const FNavigationToolSerializedItem ParentSceneItem = InParentItem->MakeSerializedItem();
 
-	for (int32 Index = 0; Index < Children.Num(); ++Index)
+	for (int32 Index = 0; Index < WeakChildren.Num(); ++Index)
 	{
-		if (const FNavigationToolItemPtr& ChildItem = Children[Index])
+		if (const FNavigationToolViewModelPtr ChildItem = WeakChildren[Index].Pin())
 		{
 			if (ChildItem->ShouldSort())
 			{
@@ -294,33 +311,39 @@ void FNavigationToolProvider::SaveSerializedTreeRecursive(const FNavigationToolI
 	}
 }
 
-void FNavigationToolProvider::LoadSerializedTree(const FNavigationToolItemPtr& InParentItem
+void FNavigationToolProvider::LoadSerializedTree(const FNavigationToolViewModelPtr& InParentItem
 	, FNavigationToolSerializedTree* const InSerializedTree)
 {
-	TArray<FNavigationToolItemPtr>& Children = InParentItem->GetChildrenMutable();
+	using namespace ItemUtils;
+
+	TArray<FNavigationToolViewModelWeakPtr>& WeakChildren = InParentItem->GetChildrenMutable();
 	
-	TArray<FNavigationToolItemPtr> Sortable;
-	TArray<FNavigationToolItemPtr> Unsortable;
-	SplitSortableAndUnsortableItems(Children, Sortable, Unsortable);
+	TArray<FNavigationToolViewModelWeakPtr> WeakSortable;
+	TArray<FNavigationToolViewModelWeakPtr> WeakUnsortable;
+	SplitSortableAndUnsortableItems(WeakChildren, WeakSortable, WeakUnsortable);
 
 	// If Scene Tree is valid, Item Sorting should be empty as this function only takes a valid Scene Tree if
 	// loaded version supports Scene Trees (i.e. when Item Sorting stops being loaded in)
 	if (InSerializedTree)
 	{
-		Sortable.Sort([InSerializedTree](const FNavigationToolItemPtr& InItemA, const FNavigationToolItemPtr& InItemB)
+		WeakSortable.Sort([InSerializedTree](const FNavigationToolViewModelWeakPtr& InWeakItemA, const FNavigationToolViewModelWeakPtr& InWeakItemB)
 			{
-				const FNavigationToolSerializedTreeNode* const NodeA = InSerializedTree->FindTreeNode(InItemA->MakeSerializedItem());
-				const FNavigationToolSerializedTreeNode* const NodeB = InSerializedTree->FindTreeNode(InItemB->MakeSerializedItem());
+				const TViewModelPtr<INavigationToolItem> ItemA = InWeakItemA.Pin();
+				const TViewModelPtr<INavigationToolItem> ItemB = InWeakItemB.Pin();
+				const FNavigationToolSerializedTreeNode* const NodeA = ItemA.IsValid()
+					? InSerializedTree->FindTreeNode(ItemA->MakeSerializedItem()) : nullptr;
+				const FNavigationToolSerializedTreeNode* const NodeB = ItemB.IsValid()
+					? InSerializedTree->FindTreeNode(ItemB->MakeSerializedItem()) : nullptr;
 				return FNavigationToolSerializedTree::CompareTreeItemOrder(NodeA, NodeB);
 			});
 	}
 
-	Children = MoveTemp(Unsortable);
-	Children.Append(MoveTemp(Sortable));
+	WeakChildren = MoveTemp(WeakUnsortable);
+	WeakChildren.Append(MoveTemp(WeakSortable));
 
-	for (FNavigationToolItemPtr& Child : Children)
+	for (FNavigationToolViewModelWeakPtr& WeakChild : WeakChildren)
 	{
-		LoadSerializedTree(Child, InSerializedTree);
+		LoadSerializedTree(WeakChild.Pin(), InSerializedTree);
 	}
 }
 
diff --git a/Engine/Plugins/Experimental/SequenceNavigator/Source/SequenceNavigator/Private/SequenceNavigatorModule.cpp b/Engine/Plugins/Experimental/SequenceNavigator/Source/SequenceNavigator/Private/SequenceNavigatorModule.cpp
--- a/Engine/Plugins/Experimental/SequenceNavigator/Source/SequenceNavigator/Private/SequenceNavigatorModule.cpp
+++ b/Engine/Plugins/Experimental/SequenceNavigator/Source/SequenceNavigator/Private/SequenceNavigatorModule.cpp
@@ -1,10 +1,29 @@
 ﻿// Copyright Epic Games, Inc. All Rights Reserved.
 
 #include "SequenceNavigatorModule.h"
+#include "Extensions/IColorExtension.h"
+#include "Extensions/IInTimeExtension.h"
+#include "Extensions/IMarkerVisibilityExtension.h"
+#include "Extensions/IOutTimeExtension.h"
+#include "Extensions/IPlayheadExtension.h"
+#include "Extensions/IRevisionControlExtension.h"
+#include "Items/INavigationToolItem.h"
+#include "MVVM/ViewModelTypeID.h"
 #include "NavigationToolCommands.h"
 #include "SequenceNavigatorLog.h"
 
 DEFINE_LOG_CATEGORY(LogSequenceNavigator);
+/*
+namespace UE::SequenceNavigator
+{
+	UE_SEQUENCER_DEFINE_VIEW_MODEL_TYPE_ID(INavigationToolItem)
+	UE_SEQUENCER_DEFINE_VIEW_MODEL_TYPE_ID(IColorExtension)
+	UE_SEQUENCER_DEFINE_VIEW_MODEL_TYPE_ID(IInTimeExtension)
+	UE_SEQUENCER_DEFINE_VIEW_MODEL_TYPE_ID(IMarkerVisibilityExtension)
+	UE_SEQUENCER_DEFINE_VIEW_MODEL_TYPE_ID(IOutTimeExtension)
+	UE_SEQUENCER_DEFINE_VIEW_MODEL_TYPE_ID(IPlayheadExtension)
+	UE_SEQUENCER_DEFINE_VIEW_MODEL_TYPE_ID(IRevisionControlExtension)
+}*/
 
 void FSequenceNavigatorModule::StartupModule()
 {
diff --git a/Engine/Plugins/Experimental/SequenceNavigator/Source/SequenceNavigator/Private/Utils/NavigationToolMiscUtils.cpp b/Engine/Plugins/Experimental/SequenceNavigator/Source/SequenceNavigator/Private/Utils/NavigationToolMiscUtils.cpp
--- a/Engine/Plugins/Experimental/SequenceNavigator/Source/SequenceNavigator/Private/Utils/NavigationToolMiscUtils.cpp
+++ b/Engine/Plugins/Experimental/SequenceNavigator/Source/SequenceNavigator/Private/Utils/NavigationToolMiscUtils.cpp
@@ -19,6 +19,8 @@
 namespace UE::SequenceNavigator
 {
 
+using namespace Sequencer;
+
 FSourceControlStatePtr FindSourceControlState(UPackage* const InPackage)
 {
 	if (!InPackage)
@@ -118,7 +120,7 @@ void FocusSequence(const INavigationTool& InTool
 	, UMovieSceneSequence& InSequence
 	, const FNavigationToolTrack& InTrackItem)
 {
-	using namespace UE::Sequencer;
+	using namespace Sequencer;
 
 	FocusSequence(InTool, InSequence);
 
@@ -145,7 +147,7 @@ void FocusSequence(const INavigationTool& InTool
 	, UMovieSceneSequence& InSequence
 	, const FNavigationToolBinding& InBindingItem)
 {
-	using namespace UE::Sequencer;
+	using namespace Sequencer;
 
 	FocusSequence(InTool, InSequence);
 
@@ -166,7 +168,7 @@ void FocusSequence(const INavigationTool& InTool
 
 void SelectSequencerBindingTrack(const ISequencer& InSequencer, const FGuid& InObjectId)
 {
-	using namespace UE::Sequencer;
+	using namespace Sequencer;
 
 	const TSharedPtr<FSequencerEditorViewModel> ViewModel = InSequencer.GetViewModel();
 	if (!ViewModel.IsValid())
@@ -214,23 +216,23 @@ void FocusSequence(const INavigationTool& InTool
 	}
 }
 
-void FocusItemInSequencer(const INavigationTool& InTool, const FNavigationToolItemPtr& InItem)
+void FocusItemInSequencer(const INavigationTool& InTool, const FNavigationToolViewModelPtr& InItem)
 {
-	if (const FNavigationToolSequence* const SequenceItem = InItem->CastTo<FNavigationToolSequence>())
+	if (const TViewModelPtr<FNavigationToolSequence> SequenceItem = InItem.ImplicitCast())
 	{
 		if (UMovieSceneSequence* const ParentSequence = SequenceItem->GetSequence())
 		{
 			FocusSequence(InTool, *ParentSequence);
 		}
 	}
-	else if (const FNavigationToolTrack* const TrackItem = InItem->CastTo<FNavigationToolTrack>())
+	else if (const TViewModelPtr<FNavigationToolTrack> TrackItem = InItem.ImplicitCast())
 	{
 		if (UMovieSceneSequence* const Sequence = TrackItem->GetSequence())
 		{
 			FocusSequence(InTool, *Sequence, *TrackItem);
 		}
 	}
-	else if (const FNavigationToolBinding* const BindingItem = InItem->CastTo<FNavigationToolBinding>())
+	else if (const TViewModelPtr<FNavigationToolBinding> BindingItem = InItem.ImplicitCast())
 	{
 		if (UMovieSceneSequence* const Sequence = BindingItem->GetSequence())
 		{
diff --git a/Engine/Plugins/Experimental/SequenceNavigator/Source/SequenceNavigator/Private/Utils/NavigationToolMiscUtils.h b/Engine/Plugins/Experimental/SequenceNavigator/Source/SequenceNavigator/Private/Utils/NavigationToolMiscUtils.h
--- a/Engine/Plugins/Experimental/SequenceNavigator/Source/SequenceNavigator/Private/Utils/NavigationToolMiscUtils.h
+++ b/Engine/Plugins/Experimental/SequenceNavigator/Source/SequenceNavigator/Private/Utils/NavigationToolMiscUtils.h
@@ -94,7 +94,7 @@ void FocusSequence(const INavigationTool& InTool
 	, const FMovieSceneMarkedFrame& InMarkedFrame);
 
 void FocusItemInSequencer(const INavigationTool& InTool
-	, const FNavigationToolItemPtr& InItem);
+	, const FNavigationToolViewModelPtr& InItem);
 
 /**
  * Resolve the sequence ID for a given sequence within the provided sequencer.
diff --git a/Engine/Plugins/Experimental/SequenceNavigator/Source/SequenceNavigator/Private/Widgets/Columns/SNavigationToolColor.cpp b/Engine/Plugins/Experimental/SequenceNavigator/Source/SequenceNavigator/Private/Widgets/Columns/SNavigationToolColor.cpp
--- a/Engine/Plugins/Experimental/SequenceNavigator/Source/SequenceNavigator/Private/Widgets/Columns/SNavigationToolColor.cpp
+++ b/Engine/Plugins/Experimental/SequenceNavigator/Source/SequenceNavigator/Private/Widgets/Columns/SNavigationToolColor.cpp
@@ -17,6 +17,8 @@
 namespace UE::SequenceNavigator
 {
 
+using namespace Sequencer;
+
 DECLARE_DELEGATE_RetVal_OneParam(FReply, FOnColorEntrySelected, FName EntryName);
 
 class SNavigationToolColorEntry : public SCompoundWidget
@@ -75,7 +77,7 @@ public:
 		return &MenuButtonStyle->Normal;
 	}
 
-	virtual FReply OnMouseButtonUp(const FGeometry& MyGeometry, const FPointerEvent& MouseEvent) override
+	virtual FReply OnMouseButtonUp(const FGeometry& InGeometry, const FPointerEvent& InPointerEvent) override
 	{
 		if (OnColorEntrySelected.IsBound())
 		{
@@ -112,7 +114,7 @@ public:
 };
 
 void SNavigationToolColor::Construct(const FArguments& InArgs
-	, const FNavigationToolItemRef& InItem
+	, const FNavigationToolViewModelPtr& InItem
 	, const TSharedRef<INavigationToolView>& InView
 	, const TSharedRef<SNavigationToolTreeRow>& InRowWidget)
 {
@@ -146,9 +148,12 @@ FLinearColor SNavigationToolColor::GetColorBlockColor() const
 {
 	ItemColor = FColor();
 
-	if (const IColorExtension* const ColorExtension = WeakItem.Pin()->CastTo<IColorExtension>())
+	if (const FNavigationToolViewModelPtr Item = WeakItem.Pin())
 	{
-		ItemColor = ColorExtension->GetColor().Get(FColor());
+		if (const TViewModelPtr<IColorExtension> ColorExtension = Item.ImplicitCast())
+		{
+			ItemColor = ColorExtension->GetColor().Get(FColor());
+		}
 	}
 
 	return ItemColor;
@@ -172,7 +177,7 @@ FSlateColor SNavigationToolColor::GetBorderColor() const
 
 FLinearColor SNavigationToolColor::GetStateColorAndOpacity() const
 {
-	const FNavigationToolItemPtr Item = WeakItem.Pin();
+	const FNavigationToolViewModelPtr Item = WeakItem.Pin();
 	if (!Item.IsValid())
 	{
 		return FLinearColor::Transparent;
@@ -184,7 +189,7 @@ FLinearColor SNavigationToolColor::GetStateColorAndOpacity() const
 		return FLinearColor::Transparent;
 	}
 
-	const bool bIsSelected = ToolView->IsItemSelected(Item.ToSharedRef());
+	const bool bIsSelected = ToolView->IsItemSelected(Item);
 
 	const TSharedPtr<SNavigationToolTreeRow> Row = WeakRowWidget.Pin();
 
@@ -216,24 +221,24 @@ FReply SNavigationToolColor::OnColorEntrySelected(const FColor& InColor) const
 	return FReply::Unhandled();
 }
 
-FReply SNavigationToolColor::OnDragDetected(const FGeometry& MyGeometry, const FPointerEvent& MouseEvent)
+FReply SNavigationToolColor::OnDragDetected(const FGeometry& InGeometry, const FPointerEvent& InPointerEvent)
 {
-	if (MouseEvent.IsMouseButtonDown(EKeys::LeftMouseButton))
+	if (InPointerEvent.IsMouseButtonDown(EKeys::LeftMouseButton))
 	{
 		return FReply::Handled().BeginDragDrop(FColorDragDropOp::New(ItemColor));
 	}
 
-	return SCompoundWidget::OnDragDetected(MyGeometry, MouseEvent);
+	return SCompoundWidget::OnDragDetected(InGeometry, InPointerEvent);
 }
 
-void SNavigationToolColor::OnDragEnter(const FGeometry& MyGeometry, const FDragDropEvent& DragDropEvent)
+void SNavigationToolColor::OnDragEnter(const FGeometry& InGeometry, const FDragDropEvent& InDragDropEvent)
 {
-	if (const TSharedPtr<FColorDragDropOp> ColorOp = DragDropEvent.GetOperationAs<FColorDragDropOp>())
+	if (const TSharedPtr<FColorDragDropOp> ColorOp = InDragDropEvent.GetOperationAs<FColorDragDropOp>())
 	{
 		OnColorEntrySelected(ColorOp->Color);
 	}
 
-	SCompoundWidget::OnDragEnter(MyGeometry, DragDropEvent);
+	SCompoundWidget::OnDragEnter(InGeometry, InDragDropEvent);
 }
 
 FReply SNavigationToolColor::OnColorMouseButtonDown(const FGeometry& InGeometry, const FPointerEvent& InEvent)
@@ -254,14 +259,14 @@ void SNavigationToolColor::OpenColorPickerDialog()
 	PickerArgs.ParentWidget = GetParentWidget();
 	PickerArgs.OnColorCommitted = FOnLinearColorValueChanged::CreateLambda([this](const FLinearColor InNewColor)
 		{
-			if (IColorExtension* const ColorExtension = WeakItem.Pin()->CastTo<IColorExtension>())
+			if (const TViewModelPtr<IColorExtension> ColorExtension = WeakItem.ImplicitPin())
 			{
 				ColorExtension->SetColor(InNewColor.ToFColor(true));
 			}
 		});
 	PickerArgs.OnColorPickerCancelled = FOnColorPickerCancelled::CreateLambda([this, PreviousColor](const FLinearColor& InNewColor)
 		{
-			if (IColorExtension* const ColorExtension = WeakItem.Pin()->CastTo<IColorExtension>())
+			if (const TViewModelPtr<IColorExtension> ColorExtension = WeakItem.ImplicitPin())
 			{
 				ColorExtension->SetColor(PreviousColor.ToFColor(true));
 			}
diff --git a/Engine/Plugins/Experimental/SequenceNavigator/Source/SequenceNavigator/Private/Widgets/Columns/SNavigationToolColor.h b/Engine/Plugins/Experimental/SequenceNavigator/Source/SequenceNavigator/Private/Widgets/Columns/SNavigationToolColor.h
--- a/Engine/Plugins/Experimental/SequenceNavigator/Source/SequenceNavigator/Private/Widgets/Columns/SNavigationToolColor.h
+++ b/Engine/Plugins/Experimental/SequenceNavigator/Source/SequenceNavigator/Private/Widgets/Columns/SNavigationToolColor.h
@@ -10,6 +10,7 @@ class SMenuAnchor;
 namespace UE::SequenceNavigator
 {
 
+class IColorExtension;
 class INavigationTool;
 class INavigationToolView;
 class SNavigationToolTreeRow;
@@ -23,7 +24,7 @@ public:
 	virtual ~SNavigationToolColor() override {}
 
 	void Construct(const FArguments& InArgs
-		, const FNavigationToolItemRef& InItem
+		, const FNavigationToolViewModelPtr& InItem
 		, const TSharedRef<INavigationToolView>& InView
 		, const TSharedRef<SNavigationToolTreeRow>& InRowWidget);
 
@@ -35,15 +36,15 @@ public:
 
 	FReply OnColorEntrySelected(const FColor& InColor) const;
 
-	virtual FReply OnDragDetected(const FGeometry& MyGeometry, const FPointerEvent& MouseEvent) override;
-	virtual void OnDragEnter(const FGeometry& MyGeometry, const FDragDropEvent& DragDropEvent) override;
+	virtual FReply OnDragDetected(const FGeometry& InGeometry, const FPointerEvent& InPointerEvent) override;
+	virtual void OnDragEnter(const FGeometry& InGeometry, const FDragDropEvent& InDragDropEvent) override;
 
 protected:
 	FReply OnColorMouseButtonDown(const FGeometry& InGeometry, const FPointerEvent& InEvent);
 
 	void OpenColorPickerDialog();
 
-	FNavigationToolItemWeakPtr WeakItem;
+	FNavigationToolViewModelWeakPtr WeakItem;
 
 	TWeakPtr<INavigationToolView> WeakView;
 
diff --git a/Engine/Plugins/Experimental/SequenceNavigator/Source/SequenceNavigator/Private/Widgets/Columns/SNavigationToolComment.cpp b/Engine/Plugins/Experimental/SequenceNavigator/Source/SequenceNavigator/Private/Widgets/Columns/SNavigationToolComment.cpp
--- a/Engine/Plugins/Experimental/SequenceNavigator/Source/SequenceNavigator/Private/Widgets/Columns/SNavigationToolComment.cpp
+++ b/Engine/Plugins/Experimental/SequenceNavigator/Source/SequenceNavigator/Private/Widgets/Columns/SNavigationToolComment.cpp
@@ -16,7 +16,7 @@ namespace UE::SequenceNavigator
 {
 
 void SNavigationToolComment::Construct(const FArguments& InArgs
-	, const FNavigationToolItemRef& InItem
+	, const FNavigationToolViewModelPtr& InItem
 	, const TSharedRef<INavigationToolView>& InView
 	, const TSharedRef<SNavigationToolTreeRow>& InRowWidget)
 {
@@ -24,12 +24,6 @@ void SNavigationToolComment::Construct(const FArguments& InArgs
 	WeakView = InView;
 	WeakRowWidget = InRowWidget;
 
-	const FNavigationToolSequence* const SequenceItem = InItem->CastTo<FNavigationToolSequence>();
-	if (!SequenceItem || !SequenceItem->GetSubSection())
-	{
-		return;
-	}
-
 	ChildSlot
 		[
 			SNew(SEditableTextBox)
@@ -45,7 +39,7 @@ SNavigationToolComment::~SNavigationToolComment()
 
 FString SNavigationToolComment::GetMetaDataComment() const
 {
-	if (UMovieSceneMetaData* const SequenceMetaData = GetSequenceItemMetaData(WeakItem.Pin()))
+	if (const UMovieSceneMetaData* const SequenceMetaData = ItemUtils::GetSequenceItemMetaData(WeakItem.Pin()))
 	{
 		return SequenceMetaData->GetNotes();
 	}
@@ -64,8 +58,16 @@ void SNavigationToolComment::OnCommentTextChanged(const FText& InNewText)
 
 void SNavigationToolComment::OnCommentTextCommitted(const FText& InNewText, const ETextCommit::Type InCommitType)
 {
-	const FNavigationToolSequence* const SequenceItem = WeakItem.Pin()->CastTo<FNavigationToolSequence>();
-	if (!SequenceItem)
+	using namespace Sequencer;
+
+	const FNavigationToolViewModelPtr Item = WeakItem.Pin();
+	if (!Item.IsValid())
+	{
+		return;
+	}
+
+	const TViewModelPtr<FNavigationToolSequence> SequenceItem = Item.ImplicitCast();
+	if (!SequenceItem.IsValid())
 	{
 		return;
 	}
