[SequenceNavigator] Refactor to use Sequencer view models and type macros

diff --git a/Engine/Plugins/Experimental/SequenceNavigator/Source/SequenceNavigator/Private/Widgets/Columns/SNavigationToolComment.h b/Engine/Plugins/Experimental/SequenceNavigator/Source/SequenceNavigator/Private/Widgets/Columns/SNavigationToolComment.h
--- a/Engine/Plugins/Experimental/SequenceNavigator/Source/SequenceNavigator/Private/Widgets/Columns/SNavigationToolComment.h
+++ b/Engine/Plugins/Experimental/SequenceNavigator/Source/SequenceNavigator/Private/Widgets/Columns/SNavigationToolComment.h
@@ -2,13 +2,16 @@
 
 #pragma once
 
-#include "NavigationToolDefines.h"
+#include "MVVM/ViewModelPtr.h"
+//#include "NavigationToolDefines.h"
+#include "Items/NavigationToolSequence.h"
 #include "Templates/UniquePtr.h"
 #include "Widgets/SCompoundWidget.h"
 
 namespace UE::SequenceNavigator
 {
 
+//class FNavigationToolSequence;
 class INavigationToolView;
 class SNavigationToolTreeRow;
 
@@ -21,11 +24,11 @@ public:
 	virtual ~SNavigationToolComment() override;
 
 	void Construct(const FArguments& InArgs
-		, const FNavigationToolItemRef& InItem
+		, const FNavigationToolViewModelPtr& InItem
 		, const TSharedRef<INavigationToolView>& InView
 		, const TSharedRef<SNavigationToolTreeRow>& InRowWidget);
 
-private:
+protected:
 	FString GetMetaDataComment() const;
 
 	FText GetCommentText() const;
@@ -34,7 +37,7 @@ private:
 
 	FText GetTransactionText() const;
 
-	FNavigationToolItemWeakPtr WeakItem;
+	FNavigationToolViewModelWeakPtr WeakItem;
 
 	TWeakPtr<INavigationToolView> WeakView;
 
diff --git a/Engine/Plugins/Experimental/SequenceNavigator/Source/SequenceNavigator/Private/Widgets/Columns/SNavigationToolDeactiveState.cpp b/Engine/Plugins/Experimental/SequenceNavigator/Source/SequenceNavigator/Private/Widgets/Columns/SNavigationToolDeactiveState.cpp
--- a/Engine/Plugins/Experimental/SequenceNavigator/Source/SequenceNavigator/Private/Widgets/Columns/SNavigationToolDeactiveState.cpp
+++ b/Engine/Plugins/Experimental/SequenceNavigator/Source/SequenceNavigator/Private/Widgets/Columns/SNavigationToolDeactiveState.cpp
@@ -3,8 +3,9 @@
 #include "SNavigationToolDeactiveState.h"
 #include "Columns/NavigationToolDeactiveStateColumn.h"
 #include "Editor.h"
-#include "Extensions/ISequenceInactivableExtension.h"
 #include "Items/INavigationToolItem.h"
+#include "Items/NavigationToolItemUtils.h"
+#include "MVVM/Extensions/IDeactivatableExtension.h"
 #include "NavigationToolView.h"
 #include "Widgets/SNavigationToolTreeRow.h"
 
@@ -13,13 +14,15 @@
 namespace UE::SequenceNavigator
 {
 
+using namespace Sequencer;
+
 class FEvaluationStateDragDropOp : public FDragDropOperation
 {
 public:
 	DRAG_DROP_OPERATOR_TYPE(FVisibilityDragDropOp, FDragDropOperation)
 
 	/** Flag which defines whether to hide destination items or not */
-	EItemSequenceInactiveState InactiveState;
+	EDeactivatableState InactiveState;
 
 	/** Undo transaction stolen from the gutter which is kept alive for the duration of the drag */
 	TUniquePtr<FScopedTransaction> UndoTransaction;
@@ -31,7 +34,7 @@ public:
 	}
 
 	/** Create a new drag and drop operation out of the specified flag */
-	static TSharedRef<FEvaluationStateDragDropOp> New(const EItemSequenceInactiveState InEvaluationState, TUniquePtr<FScopedTransaction>& ScopedTransaction)
+	static TSharedRef<FEvaluationStateDragDropOp> New(const EDeactivatableState InEvaluationState, TUniquePtr<FScopedTransaction>& ScopedTransaction)
 	{
 		TSharedRef<FEvaluationStateDragDropOp> Operation = MakeShared<FEvaluationStateDragDropOp>();
 		Operation->InactiveState = InEvaluationState;
@@ -41,38 +44,40 @@ public:
 	}
 };
 
-EItemSequenceInactiveState GetItemInactiveState(const FNavigationToolItemPtr& InItem)
+EDeactivatableState GetItemDeactivatedState(const FNavigationToolViewModelPtr& InItem)
 {
+	using namespace Sequencer;
+	using namespace ItemUtils;
+
 	if (!InItem.IsValid())
 	{
-		return EItemSequenceInactiveState::None;
+		return EDeactivatableState::None;
 	}
 
-	ISequenceInactivableExtension* const InactivatableItem = InItem->CastTo<ISequenceInactivableExtension>();
-	if (!InactivatableItem)
-	{
-		return EItemSequenceInactiveState::None;
-	}
+	const ENavigationToolCompareState State = CompareChildrenItemState<IDeactivatableExtension>(InItem,
+		[](const TViewModelPtr<IDeactivatableExtension>& InItem)
+			{
+				return InItem->IsDeactivated();
+			},
+		[](const TViewModelPtr<IDeactivatableExtension>& InItem)
+			{
+				return !InItem->IsDeactivated();
+			});
 
-	return InactivatableItem->GetInactiveState();
+	return static_cast<EDeactivatableState>(State);
 }
 
-void SetItemInactive(const FNavigationToolItemPtr& InItem, const bool bInInactive)
+void SetItemDeactivated(const FNavigationToolViewModelPtr& InItem, const bool bInInactive)
 {
-	if (!InItem.IsValid())
+	if (const TViewModelPtr<IDeactivatableExtension> DeactivatableItem = InItem.ImplicitCast())
 	{
-		return;
-	}
-
-	if (ISequenceInactivableExtension* const InactivatableItem = InItem->CastTo<ISequenceInactivableExtension>())
-	{
-		InactivatableItem->SetIsInactive(bInInactive);
+		DeactivatableItem->SetIsDeactivated(bInInactive);
 	}
 }
 
 void SNavigationToolDeactiveState::Construct(const FArguments& InArgs
 	, const TSharedRef<FNavigationToolDeactiveStateColumn>& InColumn
-	, const FNavigationToolItemRef& InItem
+	, const FNavigationToolViewModelPtr& InItem
 	, const TSharedRef<INavigationToolView>& InView
 	, const TSharedRef<SNavigationToolTreeRow>& InRowWidget)
 {
@@ -87,36 +92,41 @@ void SNavigationToolDeactiveState::Construct(const FArguments& InArgs
 		.Image(this, &SNavigationToolDeactiveState::GetBrush));
 }
 
-FReply SNavigationToolDeactiveState::OnDragDetected(const FGeometry& MyGeometry, const FPointerEvent& MouseEvent)
+FReply SNavigationToolDeactiveState::OnDragDetected(const FGeometry& InGeometry, const FPointerEvent& InPointerEvent)
 {
-	if (MouseEvent.IsMouseButtonDown(EKeys::LeftMouseButton))
+	using namespace Sequencer;
+
+	if (InPointerEvent.IsMouseButtonDown(EKeys::LeftMouseButton))
 	{
-		const EItemSequenceInactiveState NewState = (GetInactiveState() != EItemSequenceInactiveState::None)
-			? EItemSequenceInactiveState::Inactive : EItemSequenceInactiveState::None;
+		const EDeactivatableState NewState = (GetInactiveState() != EDeactivatableState::None)
+			? EDeactivatableState::Deactivated : EDeactivatableState::None;
 		return FReply::Handled().BeginDragDrop(FEvaluationStateDragDropOp::New(NewState, UndoTransaction));
 	}
+
 	return FReply::Unhandled();
 }
 
-void SNavigationToolDeactiveState::OnDragEnter(const FGeometry& MyGeometry, const FDragDropEvent& DragDropEvent)
+void SNavigationToolDeactiveState::OnDragEnter(const FGeometry& InGeometry, const FDragDropEvent& InDragDropEvent)
 {
-	const TSharedPtr<FEvaluationStateDragDropOp> DragDropOp = DragDropEvent.GetOperationAs<FEvaluationStateDragDropOp>();
+	const TSharedPtr<FEvaluationStateDragDropOp> DragDropOp = InDragDropEvent.GetOperationAs<FEvaluationStateDragDropOp>();
 	if (DragDropOp.IsValid())
 	{
-		const bool bItemInactive = (DragDropOp->InactiveState != EItemSequenceInactiveState::None);
-		SetIsInactive(bItemInactive);
+		const bool bItemInactive = (DragDropOp->InactiveState != EDeactivatableState::None);
+		SetIsDeactivated(bItemInactive);
 	}
 }
 
 FReply SNavigationToolDeactiveState::HandleClick()
 {
+	using namespace Sequencer;
+
 	if (!IsVisibilityWidgetEnabled())
 	{
 		return FReply::Unhandled();
 	}
 
 	const TSharedPtr<INavigationToolView> ToolView = WeakView.Pin();
-	const FNavigationToolItemPtr Item = WeakItem.Pin();
+	const FNavigationToolViewModelPtr Item = WeakItem.Pin();
 	const TSharedPtr<FNavigationToolDeactiveStateColumn> Column = WeakColumn.Pin();
 
 	if (!ToolView.IsValid() || !Item.IsValid() || !Column.IsValid())
@@ -127,32 +137,32 @@ FReply SNavigationToolDeactiveState::HandleClick()
 	// Open an undo transaction
 	UndoTransaction = MakeUnique<FScopedTransaction>(LOCTEXT("SetNavigationToolItemInactiveState", "Set Item Inactive State"));
 
-	const bool bNewIsInactive = (GetInactiveState() != EItemSequenceInactiveState::Inactive);
+	const bool bNewIsInactive = (GetInactiveState() != EDeactivatableState::Deactivated);
 
 	// We operate on all the selected items if the specified item is selected
-	if (ToolView->IsItemSelected(Item.ToSharedRef()))
+	if (ToolView->IsItemSelected(Item))
 	{
-		for (FNavigationToolItemPtr& SelectedItem : ToolView->GetSelectedItems())
+		for (FNavigationToolViewModelWeakPtr& WeakSelectedItem : ToolView->GetSelectedItems())
 		{
-			SetItemInactive(SelectedItem, bNewIsInactive);
+			SetItemDeactivated(WeakSelectedItem.Pin(), bNewIsInactive);
 		}
 	}
 	else
 	{
-		SetIsInactive(bNewIsInactive);
+		SetIsDeactivated(bNewIsInactive);
 	}
 
 	return FReply::Handled().DetectDrag(SharedThis(this), EKeys::LeftMouseButton);
 }
 
-FReply SNavigationToolDeactiveState::OnMouseButtonDoubleClick(const FGeometry& InMyGeometry, const FPointerEvent& InMouseEvent)
+FReply SNavigationToolDeactiveState::OnMouseButtonDoubleClick(const FGeometry& InInGeometry, const FPointerEvent& InPointerEvent)
 {
 	return HandleClick();
 }
 
-FReply SNavigationToolDeactiveState::OnMouseButtonDown(const FGeometry& MyGeometry, const FPointerEvent& MouseEvent)
+FReply SNavigationToolDeactiveState::OnMouseButtonDown(const FGeometry& InGeometry, const FPointerEvent& InPointerEvent)
 {
-	if (MouseEvent.GetEffectingButton() != EKeys::LeftMouseButton)
+	if (InPointerEvent.GetEffectingButton() != EKeys::LeftMouseButton)
 	{
 		return FReply::Unhandled();
 	}
@@ -160,9 +170,9 @@ FReply SNavigationToolDeactiveState::OnMouseButtonDown(const FGeometry& MyGeomet
 	return HandleClick();
 }
 
-FReply SNavigationToolDeactiveState::OnMouseButtonUp(const FGeometry& MyGeometry, const FPointerEvent& MouseEvent)
+FReply SNavigationToolDeactiveState::OnMouseButtonUp(const FGeometry& InGeometry, const FPointerEvent& InPointerEvent)
 {
-	if (MouseEvent.GetEffectingButton() == EKeys::LeftMouseButton)
+	if (InPointerEvent.GetEffectingButton() == EKeys::LeftMouseButton)
 	{
 		UndoTransaction.Reset();
 		return FReply::Handled();
@@ -183,7 +193,7 @@ const FSlateBrush* SNavigationToolDeactiveState::GetBrush() const
 
 FSlateColor SNavigationToolDeactiveState::GetForegroundColor() const
 {
-	const FNavigationToolItemPtr Item = WeakItem.Pin();
+	const FNavigationToolViewModelPtr Item = WeakItem.Pin();
 	if (!Item.IsValid())
 	{
 		return FStyleColors::Transparent;
@@ -194,13 +204,13 @@ FSlateColor SNavigationToolDeactiveState::GetForegroundColor() const
 	{
 		switch (GetInactiveState())
 		{
-		case EItemSequenceInactiveState::None:
+		case EDeactivatableState::None:
 			return FStyleColors::White25;
 
-		case EItemSequenceInactiveState::PartiallyInactive:
+		case EDeactivatableState::PartiallyDeactivated:
 			return FStyleColors::ForegroundHover;
 
-		case EItemSequenceInactiveState::Inactive:
+		case EDeactivatableState::Deactivated:
 			return FStyleColors::ForegroundHover;
 		}
 	}
@@ -208,13 +218,13 @@ FSlateColor SNavigationToolDeactiveState::GetForegroundColor() const
 	{
 		switch (GetInactiveState())
 		{
-		case EItemSequenceInactiveState::None:
+		case EDeactivatableState::None:
 			return FStyleColors::Transparent;
 
-		case EItemSequenceInactiveState::PartiallyInactive:
+		case EDeactivatableState::PartiallyDeactivated:
 			return FStyleColors::White25;
 
-		case EItemSequenceInactiveState::Inactive:
+		case EDeactivatableState::Deactivated:
 			return FStyleColors::Foreground;
 		}
 	}
@@ -222,14 +232,14 @@ FSlateColor SNavigationToolDeactiveState::GetForegroundColor() const
 	return FStyleColors::Transparent;
 }
 
-EItemSequenceInactiveState SNavigationToolDeactiveState::GetInactiveState() const
+EDeactivatableState SNavigationToolDeactiveState::GetInactiveState() const
 {
-	return GetItemInactiveState(WeakItem.Pin());
+	return GetItemDeactivatedState(WeakItem.Pin());
 }
 
-void SNavigationToolDeactiveState::SetIsInactive(const bool bInIsInactive)
+void SNavigationToolDeactiveState::SetIsDeactivated(const bool bInIsDeactivated)
 {
-	SetItemInactive(WeakItem.Pin(), bInIsInactive);
+	SetItemDeactivated(WeakItem.Pin(), bInIsDeactivated);
 }
 
 } // namespace UE::SequenceNavigator
diff --git a/Engine/Plugins/Experimental/SequenceNavigator/Source/SequenceNavigator/Private/Widgets/Columns/SNavigationToolDeactiveState.h b/Engine/Plugins/Experimental/SequenceNavigator/Source/SequenceNavigator/Private/Widgets/Columns/SNavigationToolDeactiveState.h
--- a/Engine/Plugins/Experimental/SequenceNavigator/Source/SequenceNavigator/Private/Widgets/Columns/SNavigationToolDeactiveState.h
+++ b/Engine/Plugins/Experimental/SequenceNavigator/Source/SequenceNavigator/Private/Widgets/Columns/SNavigationToolDeactiveState.h
@@ -7,6 +7,12 @@
 #include "Templates/UniquePtr.h"
 #include "Widgets/Images/SImage.h"
 
+namespace UE::Sequencer
+{
+	class IDeactivatableExtension;
+	enum class EDeactivatableState;
+}
+
 namespace UE::SequenceNavigator
 {
 
@@ -14,7 +20,6 @@ class FNavigationToolItem;
 class FNavigationToolDeactiveStateColumn;
 class INavigationToolView;
 class SNavigationToolTreeRow;
-enum class EItemSequenceInactiveState;
 
 /** Widget responsible for managing the visibility for a single item */
 class SNavigationToolDeactiveState : public SImage
@@ -25,22 +30,22 @@ public:
 
 	void Construct(const FArguments& InArgs
 		, const TSharedRef<FNavigationToolDeactiveStateColumn>& InColumn
-		, const FNavigationToolItemRef& InItem
+		, const FNavigationToolViewModelPtr& InItem
 		, const TSharedRef<INavigationToolView>& InView
 		, const TSharedRef<SNavigationToolTreeRow>& InRowWidget);
 
-private:
+protected:
 	/** Returns whether the widget is enabled or not */
 	virtual bool IsVisibilityWidgetEnabled() const { return true; }
 
 	virtual const FSlateBrush* GetBrush() const;
 
 	//~ Begin SWidget
-	virtual FReply OnDragDetected(const FGeometry& MyGeometry, const FPointerEvent& MouseEvent) override;
-	virtual void OnDragEnter(const FGeometry& MyGeometry, const FDragDropEvent& DragDropEvent) override;
-	virtual FReply OnMouseButtonDoubleClick(const FGeometry& InMyGeometry, const FPointerEvent& InMouseEvent) override;
-	virtual FReply OnMouseButtonDown(const FGeometry& MyGeometry, const FPointerEvent& MouseEvent) override;
-	virtual FReply OnMouseButtonUp(const FGeometry& MyGeometry, const FPointerEvent& MouseEvent) override;
+	virtual FReply OnDragDetected(const FGeometry& InGeometry, const FPointerEvent& InPointerEvent) override;
+	virtual void OnDragEnter(const FGeometry& InGeometry, const FDragDropEvent& InDragDropEvent) override;
+	virtual FReply OnMouseButtonDoubleClick(const FGeometry& InInGeometry, const FPointerEvent& InPointerEvent) override;
+	virtual FReply OnMouseButtonDown(const FGeometry& InGeometry, const FPointerEvent& InPointerEvent) override;
+	virtual FReply OnMouseButtonUp(const FGeometry& InGeometry, const FPointerEvent& InPointerEvent) override;
 	virtual void OnMouseCaptureLost(const FCaptureLostEvent& CaptureLostEvent) override;
 	//~ End SWidget
 
@@ -48,12 +53,12 @@ private:
 
 	virtual FSlateColor GetForegroundColor() const override;
 
-	EItemSequenceInactiveState GetInactiveState() const;
-	void SetIsInactive(const bool bInIsInactive);
+	Sequencer::EDeactivatableState GetInactiveState() const;
+	void SetIsDeactivated(const bool bInIsDeactivated);
 
 	TWeakPtr<FNavigationToolDeactiveStateColumn> WeakColumn;
 
-	FNavigationToolItemWeakPtr WeakItem;
+	FNavigationToolViewModelWeakPtr WeakItem;
 
 	TWeakPtr<INavigationToolView> WeakView;
 
diff --git a/Engine/Plugins/Experimental/SequenceNavigator/Source/SequenceNavigator/Private/Widgets/Columns/SNavigationToolEndFrameOffset.cpp b/Engine/Plugins/Experimental/SequenceNavigator/Source/SequenceNavigator/Private/Widgets/Columns/SNavigationToolEndFrameOffset.cpp
--- a/Engine/Plugins/Experimental/SequenceNavigator/Source/SequenceNavigator/Private/Widgets/Columns/SNavigationToolEndFrameOffset.cpp
+++ b/Engine/Plugins/Experimental/SequenceNavigator/Source/SequenceNavigator/Private/Widgets/Columns/SNavigationToolEndFrameOffset.cpp
@@ -12,8 +12,10 @@
 namespace UE::SequenceNavigator
 {
 
+using namespace Sequencer;
+
 void SNavigationToolEndFrameOffset::Construct(const FArguments& InArgs
-	, const FNavigationToolItemRef& InItem
+	, const FNavigationToolViewModelPtr& InItem
 	, const TSharedRef<INavigationToolView>& InView
 	, const TSharedRef<SNavigationToolTreeRow>& InRowWidget)
 {
@@ -21,27 +23,31 @@ void SNavigationToolEndFrameOffset::Construct(const FArguments& InArgs
 	WeakView = InView;
 	WeakRowWidget = InRowWidget;
 
-	const FNavigationToolSequence* const SequenceItem = InItem->CastTo<FNavigationToolSequence>();
-	if (!SequenceItem || !SequenceItem->GetSubSection())
-	{
-		return;
-	}
-
 	SNavigationToolTime::Construct(SNavigationToolTime::FArguments(), InItem, InView, InRowWidget);
 }
 
 double SNavigationToolEndFrameOffset::GetFrameTimeValue() const
 {
-	if (UMovieSceneSubSection* const SubSection = Cast<UMovieSceneSubSection>(GetSequenceItemSubSection(WeakItem.Pin())))
+	const TViewModelPtr<FNavigationToolSequence> SequenceItem = WeakItem.ImplicitPin();
+	if (!SequenceItem.IsValid())
 	{
-		return SubSection->Parameters.EndFrameOffset.Value;
+		return 0.;
 	}
-	return 0.;
+
+	UMovieSceneSubSection* const SubSection = SequenceItem->GetSubSection();
+	if (!SubSection)
+	{
+		return 0.;
+	}
+
+	return SubSection->Parameters.EndFrameOffset.Value;
 }
 
 void SNavigationToolEndFrameOffset::OnFrameTimeValueCommitted(const double InNewValue, const ETextCommit::Type InCommitType)
 {
-	UMovieSceneSubSection* const SubSection = Cast<UMovieSceneSubSection>(GetSequenceItemSubSection(WeakItem.Pin()));
+	using namespace ItemUtils;
+
+	UMovieSceneSubSection* const SubSection = GetSequenceItemSubSection(WeakItem.Pin());
 	if (!SubSection)
 	{
 		return;
diff --git a/Engine/Plugins/Experimental/SequenceNavigator/Source/SequenceNavigator/Private/Widgets/Columns/SNavigationToolEndFrameOffset.h b/Engine/Plugins/Experimental/SequenceNavigator/Source/SequenceNavigator/Private/Widgets/Columns/SNavigationToolEndFrameOffset.h
--- a/Engine/Plugins/Experimental/SequenceNavigator/Source/SequenceNavigator/Private/Widgets/Columns/SNavigationToolEndFrameOffset.h
+++ b/Engine/Plugins/Experimental/SequenceNavigator/Source/SequenceNavigator/Private/Widgets/Columns/SNavigationToolEndFrameOffset.h
@@ -20,16 +20,16 @@ public:
 	virtual ~SNavigationToolEndFrameOffset() override {}
 
 	void Construct(const FArguments& InArgs
-		, const FNavigationToolItemRef& InItem
+		, const FNavigationToolViewModelPtr& InItem
 		, const TSharedRef<INavigationToolView>& InView
 		, const TSharedRef<SNavigationToolTreeRow>& InRowWidget);
 
-private:
+protected:
+	//~ Begin SNavigationToolTime
 	virtual double GetFrameTimeValue() const override;
-
 	virtual void OnFrameTimeValueCommitted(const double InNewValue, const ETextCommit::Type InCommitType) override;
-
 	virtual FText GetTransactionText() const override;
+	//~ End SNavigationToolTime
 };
 
 } // namespace UE::SequenceNavigator
diff --git a/Engine/Plugins/Experimental/SequenceNavigator/Source/SequenceNavigator/Private/Widgets/Columns/SNavigationToolHBias.cpp b/Engine/Plugins/Experimental/SequenceNavigator/Source/SequenceNavigator/Private/Widgets/Columns/SNavigationToolHBias.cpp
--- a/Engine/Plugins/Experimental/SequenceNavigator/Source/SequenceNavigator/Private/Widgets/Columns/SNavigationToolHBias.cpp
+++ b/Engine/Plugins/Experimental/SequenceNavigator/Source/SequenceNavigator/Private/Widgets/Columns/SNavigationToolHBias.cpp
@@ -14,8 +14,10 @@
 namespace UE::SequenceNavigator
 {
 
+using namespace Sequencer;
+
 void SNavigationToolHBias::Construct(const FArguments& InArgs
-	, const FNavigationToolItemRef& InItem
+	, const FNavigationToolViewModelPtr& InItem
 	, const TSharedRef<INavigationToolView>& InView
 	, const TSharedRef<SNavigationToolTreeRow>& InRowWidget)
 {
@@ -23,12 +25,6 @@ void SNavigationToolHBias::Construct(const FArguments& InArgs
 	WeakView = InView;
 	WeakRowWidget = InRowWidget;
 
-	const FNavigationToolSequence* const SequenceItem = InItem->CastTo<FNavigationToolSequence>();
-	if (!SequenceItem || !SequenceItem->GetSubSection())
-	{
-		return;
-	}
-
 	ChildSlot
 		[
 			SNew(SSpinBox<int32>)
@@ -44,13 +40,11 @@ void SNavigationToolHBias::Construct(const FArguments& InArgs
 
 int32 SNavigationToolHBias::GetValue() const
 {
-	UMovieSceneSubSection* const SubSection = Cast<UMovieSceneSubSection>(GetSequenceItemSubSection(WeakItem.Pin()));
-	if (!SubSection)
+	if (const UMovieSceneSubSection* const SubSection = ItemUtils::GetSequenceItemSubSection(WeakItem.Pin()))
 	{
-		return 0;
+		return SubSection->Parameters.HierarchicalBias;
 	}
-
-	return SubSection->Parameters.HierarchicalBias;
+	return 0;
 }
 
 void SNavigationToolHBias::OnValueChanged(const int32 InNewValue)
@@ -60,7 +54,9 @@ void SNavigationToolHBias::OnValueChanged(const int32 InNewValue)
 
 void SNavigationToolHBias::OnValueCommitted(const int32 InNewValue, const ETextCommit::Type InCommitType)
 {
-	UMovieSceneSubSection* const SubSection = Cast<UMovieSceneSubSection>(GetSequenceItemSubSection(WeakItem.Pin()));
+	using namespace ItemUtils;
+
+	UMovieSceneSubSection* const SubSection = GetSequenceItemSubSection(WeakItem.Pin());
 	if (!SubSection)
 	{
 		return;
diff --git a/Engine/Plugins/Experimental/SequenceNavigator/Source/SequenceNavigator/Private/Widgets/Columns/SNavigationToolHBias.h b/Engine/Plugins/Experimental/SequenceNavigator/Source/SequenceNavigator/Private/Widgets/Columns/SNavigationToolHBias.h
--- a/Engine/Plugins/Experimental/SequenceNavigator/Source/SequenceNavigator/Private/Widgets/Columns/SNavigationToolHBias.h
+++ b/Engine/Plugins/Experimental/SequenceNavigator/Source/SequenceNavigator/Private/Widgets/Columns/SNavigationToolHBias.h
@@ -10,6 +10,7 @@
 namespace UE::SequenceNavigator
 {
 
+class FNavigationToolSequence;
 class INavigationToolView;
 class SNavigationToolTreeRow;
 
@@ -22,11 +23,11 @@ public:
 	virtual ~SNavigationToolHBias() override {}
 
 	void Construct(const FArguments& InArgs
-		, const FNavigationToolItemRef& InItem
+		, const FNavigationToolViewModelPtr& InItem
 		, const TSharedRef<INavigationToolView>& InView
 		, const TSharedRef<SNavigationToolTreeRow>& InRowWidget);
 
-private:
+protected:
 	int32 GetValue() const;
 	void OnValueChanged(const int32 InNewValue);
 	void OnValueCommitted(const int32 InNewValue, const ETextCommit::Type InCommitType);
@@ -35,7 +36,7 @@ private:
 
 	FText GetTransactionText() const;
 
-	FNavigationToolItemWeakPtr WeakItem;
+	FNavigationToolViewModelWeakPtr WeakItem;
 
 	TWeakPtr<INavigationToolView> WeakView;
 
diff --git a/Engine/Plugins/Experimental/SequenceNavigator/Source/SequenceNavigator/Private/Widgets/Columns/SNavigationToolInTime.cpp b/Engine/Plugins/Experimental/SequenceNavigator/Source/SequenceNavigator/Private/Widgets/Columns/SNavigationToolInTime.cpp
--- a/Engine/Plugins/Experimental/SequenceNavigator/Source/SequenceNavigator/Private/Widgets/Columns/SNavigationToolInTime.cpp
+++ b/Engine/Plugins/Experimental/SequenceNavigator/Source/SequenceNavigator/Private/Widgets/Columns/SNavigationToolInTime.cpp
@@ -12,8 +12,10 @@
 namespace UE::SequenceNavigator
 {
 
+using namespace Sequencer;
+
 void SNavigationToolInTime::Construct(const FArguments& InArgs
-	, const FNavigationToolItemRef& InItem
+	, const FNavigationToolViewModelPtr& InItem
 	, const TSharedRef<INavigationToolView>& InView
 	, const TSharedRef<SNavigationToolTreeRow>& InRowWidget)
 {
@@ -21,19 +23,9 @@ void SNavigationToolInTime::Construct(const FArguments& InArgs
 	WeakView = InView;
 	WeakRowWidget = InRowWidget;
 
-	if (!InItem->IsA<IInTimeExtension>())
-	{
-		return;
-	}
-
 	SNavigationToolTime::Construct(SNavigationToolTime::FArguments(), InItem, InView, InRowWidget);
 }
 
-IInTimeExtension* SNavigationToolInTime::GetInTimeExtension() const
-{
-	return WeakItem.IsValid() ? WeakItem.Pin()->CastTo<IInTimeExtension>() : nullptr;
-}
-
 FName SNavigationToolInTime::GetStyleName() const
 {
 	return TEXT("SpinBox.InTime");
@@ -41,17 +33,17 @@ FName SNavigationToolInTime::GetStyleName() const
 
 double SNavigationToolInTime::GetFrameTimeValue() const
 {
-	if (const IInTimeExtension* const InTimeExtension = GetInTimeExtension())
+	if (const TViewModelPtr<IInTimeExtension> Item = WeakItem.ImplicitPin())
 	{
-		return InTimeExtension->GetInTime().Value;
+		return Item->GetInTime().Value;
 	}
-	return 0;
+	return 0.;
 }
 
 void SNavigationToolInTime::OnFrameTimeValueCommitted(const double InNewValue, const ETextCommit::Type InCommitType)
 {
-	IInTimeExtension* const InTimeExtension = GetInTimeExtension();
-	if (!InTimeExtension)
+	const TViewModelPtr<IInTimeExtension> Item = WeakItem.ImplicitPin();
+	if (!Item.IsValid())
 	{
 		return;
 	}
@@ -59,7 +51,7 @@ void SNavigationToolInTime::OnFrameTimeValueCommitted(const double InNewValue, c
 	const bool bShouldTransact = !UndoTransaction.IsValid() && (InCommitType == ETextCommit::OnEnter);
 	const FScopedTransaction Transaction(GetTransactionText(), bShouldTransact);
 
-	return InTimeExtension->SetInTime(static_cast<int32>(InNewValue));
+	return Item->SetInTime(static_cast<int32>(InNewValue));
 }
 
 FText SNavigationToolInTime::GetTransactionText() const
diff --git a/Engine/Plugins/Experimental/SequenceNavigator/Source/SequenceNavigator/Private/Widgets/Columns/SNavigationToolInTime.h b/Engine/Plugins/Experimental/SequenceNavigator/Source/SequenceNavigator/Private/Widgets/Columns/SNavigationToolInTime.h
--- a/Engine/Plugins/Experimental/SequenceNavigator/Source/SequenceNavigator/Private/Widgets/Columns/SNavigationToolInTime.h
+++ b/Engine/Plugins/Experimental/SequenceNavigator/Source/SequenceNavigator/Private/Widgets/Columns/SNavigationToolInTime.h
@@ -8,7 +8,7 @@
 namespace UE::SequenceNavigator
 {
 
-class IInTimeExtension;
+class FNavigationToolSequence;
 class INavigationToolView;
 class SNavigationToolTreeRow;
 
@@ -21,20 +21,17 @@ public:
 	virtual ~SNavigationToolInTime() override {}
 
 	void Construct(const FArguments& InArgs
-		, const FNavigationToolItemRef& InItem
+		, const FNavigationToolViewModelPtr& InItem
 		, const TSharedRef<INavigationToolView>& InView
 		, const TSharedRef<SNavigationToolTreeRow>& InRowWidget);
 
-private:
-	IInTimeExtension* GetInTimeExtension() const;
-
+protected:
+	//~ Begin SNavigationToolTime
 	virtual FName GetStyleName() const override;
-
 	virtual double GetFrameTimeValue() const override;
-
 	virtual void OnFrameTimeValueCommitted(const double InNewValue, const ETextCommit::Type InCommitType) override;
-
 	virtual FText GetTransactionText() const override;
+	//~ End SNavigationToolTime
 };
 
 } // namespace UE::SequenceNavigator
diff --git a/Engine/Plugins/Experimental/SequenceNavigator/Source/SequenceNavigator/Private/Widgets/Columns/SNavigationToolItemChip.cpp b/Engine/Plugins/Experimental/SequenceNavigator/Source/SequenceNavigator/Private/Widgets/Columns/SNavigationToolItemChip.cpp
--- a/Engine/Plugins/Experimental/SequenceNavigator/Source/SequenceNavigator/Private/Widgets/Columns/SNavigationToolItemChip.cpp
+++ b/Engine/Plugins/Experimental/SequenceNavigator/Source/SequenceNavigator/Private/Widgets/Columns/SNavigationToolItemChip.cpp
@@ -17,9 +17,10 @@ namespace UE::SequenceNavigator
 {
 
 void SNavigationToolItemChip::Construct(const FArguments& InArgs
-	, const FNavigationToolItemRef& InItem
+	, const FNavigationToolViewModelPtr& InItem
 	, const TSharedPtr<INavigationToolView>& InView)
 {
+	check(InItem.IsValid());
 	WeakItem = InItem;
 	WeakView = InView;
 
@@ -28,6 +29,7 @@ void SNavigationToolItemChip::Construct(const FArguments& InArgs
 	OnValidDragOver = InArgs._OnValidDragOver;
 
 	constexpr float ChipSize = 14.f;
+	const TSharedRef<INavigationToolItem> ItemRef = InItem.ToSharedRef();
 
 	SBorder::Construct(SBorder::FArguments()
 		.Padding(2.f, 1.f)
@@ -41,7 +43,7 @@ void SNavigationToolItemChip::Construct(const FArguments& InArgs
 				.Stretch(EStretch::ScaleToFit)
 				[
 					SNew(SImage)
-					.Image(InItem, &INavigationToolItem::GetIconBrush)
+					.Image(ItemRef, &INavigationToolItem::GetIconBrush)
 					.ColorAndOpacity(this, &SNavigationToolItemChip::GetIconColor)
 				]
 			]
@@ -81,41 +83,41 @@ int32 SNavigationToolItemChip::OnPaint(const FPaintArgs& InArgs, const FGeometry
 	return LayerId;
 }
 
-FReply SNavigationToolItemChip::OnMouseButtonDown(const FGeometry& InGeometry, const FPointerEvent& InMouseEvent)
+FReply SNavigationToolItemChip::OnMouseButtonDown(const FGeometry& InGeometry, const FPointerEvent& InPointerEvent)
 {
-	if (InMouseEvent.GetEffectingButton() == EKeys::LeftMouseButton)
+	if (InPointerEvent.GetEffectingButton() == EKeys::LeftMouseButton)
 	{
 		Press();
 
 		const TSharedRef<SNavigationToolItemChip> This = SharedThis(this);
 
 		return FReply::Handled()
-			.DetectDrag(This, InMouseEvent.GetEffectingButton())
+			.DetectDrag(This, InPointerEvent.GetEffectingButton())
 			.CaptureMouse(This)
 			.SetUserFocus(This, EFocusCause::Mouse);
 	}
 	return FReply::Unhandled();
 }
 
-FReply SNavigationToolItemChip::OnMouseButtonUp(const FGeometry& InGeometry, const FPointerEvent& InMouseEvent)
+FReply SNavigationToolItemChip::OnMouseButtonUp(const FGeometry& InGeometry, const FPointerEvent& InPointerEvent)
 {
 	FReply Reply = FReply::Unhandled();
 
-	if (IsPressed() && InMouseEvent.GetEffectingButton() == EKeys::LeftMouseButton)
+	if (IsPressed() && InPointerEvent.GetEffectingButton() == EKeys::LeftMouseButton)
 	{
 		Release();
 
 		if (OnItemChipClicked.IsBound())
 		{
 			bool bEventOverButton = IsHovered();
-			if (!bEventOverButton && InMouseEvent.IsTouchEvent())
+			if (!bEventOverButton && InPointerEvent.IsTouchEvent())
 			{
-				bEventOverButton = InGeometry.IsUnderLocation(InMouseEvent.GetScreenSpacePosition());
+				bEventOverButton = InGeometry.IsUnderLocation(InPointerEvent.GetScreenSpacePosition());
 			}
 
 			if (bEventOverButton && HasMouseCapture())
 			{
-				Reply = OnItemChipClicked.Execute(WeakItem.Pin(), InMouseEvent);
+				Reply = OnItemChipClicked.Execute(WeakItem.Pin(), InPointerEvent);
 			}
 		}
 
@@ -141,9 +143,9 @@ void SNavigationToolItemChip::OnMouseCaptureLost(const FCaptureLostEvent& InCapt
 	Release();
 }
 
-/*FReply SNavigationToolItemChip::OnDragDetected(const FGeometry& InGeometry, const FPointerEvent& InMouseEvent)
+/*FReply SNavigationToolItemChip::OnDragDetected(const FGeometry& InGeometry, const FPointerEvent& InPointerEvent)
 {
-	const FNavigationToolItemPtr Item = WeakItem.Pin();
+	const FNavigationToolViewModelPtr Item = WeakItem.Pin();
 	const TSharedPtr<INavigationToolView> ToolView = WeakView.Pin();
 
 	if (!Item.IsValid() || !ToolView.IsValid())
@@ -152,12 +154,12 @@ void SNavigationToolItemChip::OnMouseCaptureLost(const FCaptureLostEvent& InCapt
 		return FReply::Unhandled();
 	}
 
-	return ToolView->OnDragDetected(InGeometry, InMouseEvent, Item);
+	return ToolView->OnDragDetected(InGeometry, InPointerEvent, Item);
 }
 
 FReply SNavigationToolItemChip::OnDragOver(const FGeometry& InGeometry, const FDragDropEvent& InDragDropEvent)
 {
-	const FNavigationToolItemPtr Item = WeakItem.Pin();
+	const FNavigationToolViewModelPtr Item = WeakItem.Pin();
 	const TSharedPtr<INavigationToolView> ToolView = WeakView.Pin();
 
 	if (!Item.IsValid() || !ToolView.IsValid())
@@ -185,7 +187,7 @@ FReply SNavigationToolItemChip::OnDrop(const FGeometry& InGeometry, const FDragD
 {
 	ItemDropZone = TOptional<EItemDropZone>();
 
-	const FNavigationToolItemPtr Item = WeakItem.Pin();
+	const FNavigationToolViewModelPtr Item = WeakItem.Pin();
 	const TSharedPtr<INavigationToolView> ToolView = WeakView.Pin();
 
 	if (!Item.IsValid() || !ToolView.IsValid())
@@ -209,11 +211,19 @@ void SNavigationToolItemChip::OnDragLeave(const FDragDropEvent& InDragDropEvent)
 
 bool SNavigationToolItemChip::IsSelected() const
 {
-	if (WeakView.IsValid() && WeakItem.IsValid())
+	const TSharedPtr<INavigationToolView> View = WeakView.Pin();
+	if (!View.IsValid())
 	{
-		return WeakView.Pin()->IsItemSelected(WeakItem.Pin());
+		return false;
 	}
-	return false;
+
+	const FNavigationToolViewModelPtr Item = WeakItem.Pin();
+	if (!View.IsValid())
+	{
+		return false;
+	}
+
+	return View->IsItemSelected(Item);
 }
 
 const FSlateBrush* SNavigationToolItemChip::GetItemBackgroundBrush() const
@@ -279,7 +289,7 @@ FSlateColor SNavigationToolItemChip::GetIconColor() const
 
 FText SNavigationToolItemChip::GetTooltipText() const
 {
-	const FNavigationToolItemPtr Item = WeakItem.Pin();
+	const FNavigationToolViewModelPtr Item = WeakItem.Pin();
 	if (!Item.IsValid())
 	{
 		return FText::GetEmpty();
diff --git a/Engine/Plugins/Experimental/SequenceNavigator/Source/SequenceNavigator/Private/Widgets/Columns/SNavigationToolItemChip.h b/Engine/Plugins/Experimental/SequenceNavigator/Source/SequenceNavigator/Private/Widgets/Columns/SNavigationToolItemChip.h
--- a/Engine/Plugins/Experimental/SequenceNavigator/Source/SequenceNavigator/Private/Widgets/Columns/SNavigationToolItemChip.h
+++ b/Engine/Plugins/Experimental/SequenceNavigator/Source/SequenceNavigator/Private/Widgets/Columns/SNavigationToolItemChip.h
@@ -25,7 +25,7 @@ class INavigationToolView;
 class SNavigationToolItemChip : public SBorder
 {
 public:
-	DECLARE_DELEGATE_RetVal_TwoParams(FReply, FOnItemChipClicked, const FNavigationToolItemPtr& InItem, const FPointerEvent& InMouseEvent);
+	DECLARE_DELEGATE_RetVal_TwoParams(FReply, FOnItemChipClicked, const FNavigationToolViewModelPtr& InItem, const FPointerEvent& InPointerEvent);
 
 	SLATE_BEGIN_ARGS(SNavigationToolItemChip) {}
 		SLATE_EVENT(FOnItemChipClicked, OnItemChipClicked)
@@ -34,17 +34,17 @@ public:
 	SLATE_END_ARGS()
 
 	void Construct(const FArguments& InArgs
-		, const FNavigationToolItemRef& InItem
+		, const FNavigationToolViewModelPtr& InItem
 		, const TSharedPtr<INavigationToolView>& InView);
 
 	//~ Begin SWidget
 	virtual int32 OnPaint(const FPaintArgs& InArgs, const FGeometry& InAllottedGeometry, const FSlateRect& InCullingRect
 		, FSlateWindowElementList& OutDrawElements, int32 LayerId
 		, const FWidgetStyle& InWidgetStyle, bool bParentEnabled) const override;
-	virtual FReply OnMouseButtonDown(const FGeometry& InGeometry, const FPointerEvent& InMouseEvent) override;
-	virtual FReply OnMouseButtonUp(const FGeometry& InGeometry, const FPointerEvent& InMouseEvent) override;
+	virtual FReply OnMouseButtonDown(const FGeometry& InGeometry, const FPointerEvent& InPointerEvent) override;
+	virtual FReply OnMouseButtonUp(const FGeometry& InGeometry, const FPointerEvent& InPointerEvent) override;
 	virtual void OnMouseCaptureLost(const FCaptureLostEvent& InCaptureLostEvent) override;
-	/*virtual FReply OnDragDetected(const FGeometry& InGeometry, const FPointerEvent& InMouseEvent) override;
+	/*virtual FReply OnDragDetected(const FGeometry& InGeometry, const FPointerEvent& InPointerEvent) override;
 	virtual FReply OnDragOver(const FGeometry& InGeometry, const FDragDropEvent& InDragDropEvent) override;
 	virtual FReply OnDrop(const FGeometry& InGeometry, const FDragDropEvent& InDragDropEvent) override;
 	virtual void OnDragLeave(const FDragDropEvent& InDragDropEvent) override;*/
@@ -54,7 +54,7 @@ public:
 	const FSlateBrush* GetItemBackgroundBrush() const;
 	const FSlateBrush* GetDropIndicatorBrush(EItemDropZone InItemDropZone) const;
 
-private:
+protected:
 	bool IsPressed() const { return bIsPressed; }
 	void Press();
 	void Release();
@@ -65,7 +65,7 @@ private:
 
 	FText GetTooltipText() const;
 
-	FNavigationToolItemWeakPtr WeakItem;
+	FNavigationToolViewModelWeakPtr WeakItem;
 
 	TWeakPtr<INavigationToolView> WeakView;
 
diff --git a/Engine/Plugins/Experimental/SequenceNavigator/Source/SequenceNavigator/Private/Widgets/Columns/SNavigationToolItemList.cpp b/Engine/Plugins/Experimental/SequenceNavigator/Source/SequenceNavigator/Private/Widgets/Columns/SNavigationToolItemList.cpp
--- a/Engine/Plugins/Experimental/SequenceNavigator/Source/SequenceNavigator/Private/Widgets/Columns/SNavigationToolItemList.cpp
+++ b/Engine/Plugins/Experimental/SequenceNavigator/Source/SequenceNavigator/Private/Widgets/Columns/SNavigationToolItemList.cpp
@@ -2,6 +2,7 @@
 
 #include "SNavigationToolItemList.h"
 #include "Items/NavigationToolItem.h"
+#include "NavigationToolSettings.h"
 #include "NavigationToolView.h"
 #include "Widgets/Columns/SNavigationToolItemChip.h"
 #include "Widgets/Layout/SScrollBox.h"
@@ -13,8 +14,10 @@
 namespace UE::SequenceNavigator
 {
 
+using namespace Sequencer;
+
 void SNavigationToolItemList::Construct(const FArguments& InArgs
-	, const FNavigationToolItemRef& InItem
+	, const FNavigationToolViewModelPtr& InItem
 	, const TSharedRef<INavigationToolView>& InView
 	, const TSharedRef<SNavigationToolTreeRow>& InRowWidget)
 {
@@ -39,13 +42,16 @@ void SNavigationToolItemList::Construct(const FArguments& InArgs
 
 SNavigationToolItemList::~SNavigationToolItemList()
 {
-	if (WeakView.IsValid())
+	using namespace Sequencer;
+
+	if (const TSharedPtr<INavigationToolView> ToolView = WeakView.Pin())
 	{
-		WeakView.Pin()->GetOnToolViewRefreshed().RemoveAll(this);
+		ToolView->GetOnToolViewRefreshed().RemoveAll(this);
 	}
-	if (WeakItem.IsValid())
+
+	if (const FNavigationToolViewModelPtr Item = WeakItem.Pin())
 	{
-		WeakItem.Pin()->OnExpansionChanged().RemoveAll(this);
+		Item->OnExpansionChanged().RemoveAll(this);
 	}
 }
 
@@ -59,7 +65,7 @@ void SNavigationToolItemList::Refresh()
 	ItemListBox->ClearChildren();
 	WeakChildItems.Reset();
 
-	const FNavigationToolItemPtr ParentItem = WeakItem.Pin();
+	const FNavigationToolViewModelPtr ParentItem = WeakItem.Pin();
 	const TSharedPtr<SNavigationToolTreeRow> TreeRow = WeakRowWidget.Pin();
 	const TSharedPtr<FNavigationToolView> ToolView = StaticCastSharedPtr<FNavigationToolView>(WeakView.Pin());
 
@@ -68,44 +74,45 @@ void SNavigationToolItemList::Refresh()
 		return;
 	}
 
-	TSet<FNavigationToolItemPtr> DisallowedItems;
+	TSet<FNavigationToolViewModelWeakPtr> WeakDisallowedItems;
 	{
-		TArray<FNavigationToolItemPtr> ItemsToDisallow;
+		TArray<FNavigationToolViewModelWeakPtr> WeakItemsToDisallow;
 
-		// First get the Children in the Navigation Tool View to see if we need to Disallow some of the Items
+		// First get the children in the Navigation Tool View to see if we need to disallow some of the items
 		// to avoid redundancy or unnecessarily showing other item's children (e.g. an actor's)
-		static const TSet<FNavigationToolItemPtr> EmptySet;
-		ToolView->GetChildrenOfItem(ParentItem, ItemsToDisallow, ENavigationToolItemViewMode::ItemTree, EmptySet);
+		static const TSet<FNavigationToolViewModelWeakPtr> WeakEmptySet;
+		ToolView->GetChildrenOfItem(ParentItem, WeakItemsToDisallow, ENavigationToolItemViewMode::ItemTree, WeakEmptySet);
 
 		// If Parent item is Collapsed, only disallow items that are top levels (as these should deal with their own item list)
 		// Items that can't be top level should be visualized by the Parent Item when Collapsed
 		if (!EnumHasAllFlags(ToolView->GetViewItemFlags(ParentItem), ENavigationToolItemFlags::Expanded))
 		{
-			// keep only top level items
-			ItemsToDisallow.RemoveAllSwap([](const FNavigationToolItemPtr& Child)
+			// Keep only top level items
+			WeakItemsToDisallow.RemoveAllSwap([](const FNavigationToolViewModelWeakPtr& InWeakChild)
 				{
+					const TViewModelPtr<INavigationToolItem> Child = InWeakChild.Pin();
 					return !Child.IsValid() || !Child->CanBeTopLevel();
 				}
 				, EAllowShrinking::No);
 		}
 
-		DisallowedItems.Append(MoveTemp(ItemsToDisallow));
+		WeakDisallowedItems.Append(MoveTemp(WeakItemsToDisallow));
 	}
 
-	TArray<FNavigationToolItemPtr> Children;
-	ToolView->GetChildrenOfItem(ParentItem, Children, ENavigationToolItemViewMode::HorizontalItemList, DisallowedItems);
+	TArray<FNavigationToolViewModelWeakPtr> WeakChildren;
+	ToolView->GetChildrenOfItem(ParentItem, WeakChildren, ENavigationToolItemViewMode::HorizontalItemList, WeakDisallowedItems);
 
 	// TArray<>::Pop will be removing the item from the end which results in the Item List being in reverse
 	// so instead of immediately adding it to the slot , we will add it to this array and reverse it
-	TArray<FNavigationToolItemPtr> ItemsToAdd;
+	TArray<FNavigationToolViewModelPtr> ItemsToAdd;
 
 	const INavigationToolView& ToolViewConstRef = *ToolView;
 
-	while (!Children.IsEmpty())
+	while (!WeakChildren.IsEmpty())
 	{
-		const FNavigationToolItemPtr Child = Children.Pop();
+		const FNavigationToolViewModelPtr Child = WeakChildren.Pop().ImplicitPin();
 
-		if (!Child.IsValid() || DisallowedItems.Contains(Child))
+		if (!Child.IsValid() || WeakDisallowedItems.Contains(Child))
 		{
 			continue;
 		}
@@ -116,7 +123,7 @@ void SNavigationToolItemList::Refresh()
 		}
 		else
 		{
-			ToolView->GetChildrenOfItem(Child, Children, ENavigationToolItemViewMode::HorizontalItemList, DisallowedItems);
+			ToolView->GetChildrenOfItem(Child, WeakChildren, ENavigationToolItemViewMode::HorizontalItemList, WeakDisallowedItems);
 		}
 	}
 
@@ -124,12 +131,14 @@ void SNavigationToolItemList::Refresh()
 
 	for (int32 ItemIndex = ItemsToAdd.Num() - 1; ItemIndex >= 0; --ItemIndex)
 	{
-		const FNavigationToolItemPtr& Item = ItemsToAdd[ItemIndex];
+		const FNavigationToolViewModelPtr& Item = ItemsToAdd[ItemIndex];
+
 		WeakChildItems.Add(Item);
+
 		ItemListBox->AddSlot()
 			.Padding(0.f, 1.f)
 			[
-				SNew(SNavigationToolItemChip, Item.ToSharedRef(), ToolView)
+				SNew(SNavigationToolItemChip, Item, ToolView)
 				.ChipStyle(TreeRow->GetStyle())
 				.OnItemChipClicked(this, &SNavigationToolItemList::OnItemChipSelected)
 				.OnValidDragOver(this, &SNavigationToolItemList::OnItemChipValidDragOver)
@@ -137,17 +146,17 @@ void SNavigationToolItemList::Refresh()
 	}
 }
 
-FReply SNavigationToolItemList::OnItemChipSelected(const FNavigationToolItemPtr& InItem, const FPointerEvent& InMouseEvent)
+FReply SNavigationToolItemList::OnItemChipSelected(const FNavigationToolViewModelPtr& InItem, const FPointerEvent& InPointerEvent)
 {
 	const TSharedPtr<INavigationToolView> ToolView = WeakView.Pin();
-	const FNavigationToolItemPtr ParentItem = WeakItem.Pin();
+	const FNavigationToolViewModelPtr ParentItem = WeakItem.Pin();
 
 	if (!InItem.IsValid() || !ToolView.IsValid() || !ParentItem.IsValid())
 	{
 		return FReply::Unhandled();
 	}
 
-	if (InMouseEvent.IsAltDown())
+	if (InPointerEvent.IsAltDown())
 	{
 		constexpr ENavigationToolItemSelectionFlags SelectionFlags = ENavigationToolItemSelectionFlags::SignalSelectionChange;
 		ToolView->SelectItems({ InItem }, SelectionFlags);
diff --git a/Engine/Plugins/Experimental/SequenceNavigator/Source/SequenceNavigator/Private/Widgets/Columns/SNavigationToolItemList.h b/Engine/Plugins/Experimental/SequenceNavigator/Source/SequenceNavigator/Private/Widgets/Columns/SNavigationToolItemList.h
--- a/Engine/Plugins/Experimental/SequenceNavigator/Source/SequenceNavigator/Private/Widgets/Columns/SNavigationToolItemList.h
+++ b/Engine/Plugins/Experimental/SequenceNavigator/Source/SequenceNavigator/Private/Widgets/Columns/SNavigationToolItemList.h
@@ -23,7 +23,7 @@ public:
 	virtual ~SNavigationToolItemList() override;
 
 	void Construct(const FArguments& InArgs
-		, const FNavigationToolItemRef& InItem
+		, const FNavigationToolViewModelPtr& InItem
 		, const TSharedRef<INavigationToolView>& InView
 		, const TSharedRef<SNavigationToolTreeRow>& InRowWidget);
 
@@ -31,18 +31,18 @@ public:
 
 	void Refresh();
 
-	FReply OnItemChipSelected(const FNavigationToolItemPtr& InItem, const FPointerEvent& InMouseEvent);
+	FReply OnItemChipSelected(const FNavigationToolViewModelPtr& InItem, const FPointerEvent& InPointerEvent);
 
 	FReply OnItemChipValidDragOver(const FGeometry& InGeometry, const FDragDropEvent& InDragDropEvent);
 
-private:
-	FNavigationToolItemWeakPtr WeakItem;
+protected:
+	FNavigationToolViewModelWeakPtr WeakItem;
 
 	TWeakPtr<INavigationToolView> WeakView;
 
 	TWeakPtr<SNavigationToolTreeRow> WeakRowWidget;
 
-	TArray<FNavigationToolItemWeakPtr> WeakChildItems;
+	TArray<FNavigationToolViewModelWeakPtr> WeakChildItems;
 
 	TSharedPtr<SScrollBox> ItemListBox;
 };
diff --git a/Engine/Plugins/Experimental/SequenceNavigator/Source/SequenceNavigator/Private/Widgets/Columns/SNavigationToolLabelComponent.cpp b/Engine/Plugins/Experimental/SequenceNavigator/Source/SequenceNavigator/Private/Widgets/Columns/SNavigationToolLabelComponent.cpp
--- a/Engine/Plugins/Experimental/SequenceNavigator/Source/SequenceNavigator/Private/Widgets/Columns/SNavigationToolLabelComponent.cpp
+++ b/Engine/Plugins/Experimental/SequenceNavigator/Source/SequenceNavigator/Private/Widgets/Columns/SNavigationToolLabelComponent.cpp
@@ -11,7 +11,7 @@ namespace UE::SequenceNavigator
 {
 
 void SNavigationToolLabelComponent::Construct(const FArguments& InArgs
-	, const TSharedRef<FNavigationToolComponent>& InItem
+	, const FNavigationToolViewModelPtr& InItem
 	, const TSharedRef<SNavigationToolTreeRow>& InRowWidget)
 {
 	SNavigationToolLabelItem::Construct(SNavigationToolLabelItem::FArguments(), InItem, InRowWidget);
diff --git a/Engine/Plugins/Experimental/SequenceNavigator/Source/SequenceNavigator/Private/Widgets/Columns/SNavigationToolLabelComponent.h b/Engine/Plugins/Experimental/SequenceNavigator/Source/SequenceNavigator/Private/Widgets/Columns/SNavigationToolLabelComponent.h
--- a/Engine/Plugins/Experimental/SequenceNavigator/Source/SequenceNavigator/Private/Widgets/Columns/SNavigationToolLabelComponent.h
+++ b/Engine/Plugins/Experimental/SequenceNavigator/Source/SequenceNavigator/Private/Widgets/Columns/SNavigationToolLabelComponent.h
@@ -2,6 +2,7 @@
 
 #pragma once
 
+#include "NavigationToolDefines.h"
 #include "SNavigationToolLabelItem.h"
 
 struct FInlineEditableTextBlockStyle;
@@ -19,7 +20,7 @@ public:
 	SLATE_END_ARGS()
 
 	void Construct(const FArguments& InArgs
-		, const TSharedRef<FNavigationToolComponent>& InItem
+		, const FNavigationToolViewModelPtr& InItem
 		, const TSharedRef<SNavigationToolTreeRow>& InRowWidget);
 
 	//~ Begin SNavigationToolLabelItem
diff --git a/Engine/Plugins/Experimental/SequenceNavigator/Source/SequenceNavigator/Private/Widgets/Columns/SNavigationToolLabelItem.cpp b/Engine/Plugins/Experimental/SequenceNavigator/Source/SequenceNavigator/Private/Widgets/Columns/SNavigationToolLabelItem.cpp
--- a/Engine/Plugins/Experimental/SequenceNavigator/Source/SequenceNavigator/Private/Widgets/Columns/SNavigationToolLabelItem.cpp
+++ b/Engine/Plugins/Experimental/SequenceNavigator/Source/SequenceNavigator/Private/Widgets/Columns/SNavigationToolLabelItem.cpp
@@ -17,13 +17,18 @@
 namespace UE::SequenceNavigator
 {
 
+using namespace Sequencer;
+
 void SNavigationToolLabelItem::Construct(const FArguments& InArgs
-	, const FNavigationToolItemRef& InItem
+	, const FNavigationToolViewModelPtr& InItem
 	, const TSharedRef<SNavigationToolTreeRow>& InRowWidget)
 {
+	check(InItem.IsValid());
 	WeakItem = InItem;
 	WeakView = InRowWidget->GetToolView();
 
+	const TSharedRef<INavigationToolItem> ItemRef = InItem.ToSharedRef();
+
 	InItem->OnRenameAction().AddSP(this, &SNavigationToolLabelItem::OnRenameAction);
 
 	ChildSlot
@@ -58,9 +63,9 @@ void SNavigationToolLabelItem::Construct(const FArguments& InArgs
 					[
 						SNew(SImage)
 						.IsEnabled(this, &SNavigationToolLabelItem::IsItemEnabled)
-						.ColorAndOpacity(InItem, &INavigationToolItem::GetIconColor)
-						.Image(InItem, &INavigationToolItem::GetIconBrush)
-						.ToolTipText(InItem, &INavigationToolItem::GetIconTooltipText)
+						.ColorAndOpacity(ItemRef, &INavigationToolItem::GetIconColor)
+						.Image(ItemRef, &INavigationToolItem::GetIconBrush)
+						.ToolTipText(ItemRef, &INavigationToolItem::GetIconTooltipText)
 					]
 				]
 				+ SHorizontalBox::Slot()
@@ -88,13 +93,13 @@ void SNavigationToolLabelItem::Construct(const FArguments& InArgs
 
 bool SNavigationToolLabelItem::IsReadOnly() const
 {
-	const FNavigationToolItemPtr Item = WeakItem.Pin();
+	const FNavigationToolViewModelPtr Item = WeakItem.Pin();
 	if (!Item.IsValid())
 	{
 		return true;
 	}
 
-	if (const IRenameableExtension* const RenameableExtension = Item->CastTo<IRenameableExtension>())
+	if (const TViewModelPtr<IRenameableExtension> RenameableExtension = Item.ImplicitCast())
 	{
 		if (!RenameableExtension->CanRename())
 		{
@@ -113,16 +118,24 @@ bool SNavigationToolLabelItem::IsReadOnly() const
 
 bool SNavigationToolLabelItem::IsItemEnabled() const
 {
-	if (WeakItem.IsValid() && WeakView.IsValid())
+	const TSharedPtr<INavigationToolView> View = WeakView.Pin();
+	if (!View.IsValid())
+	{
+		return false;
+	}
+
+	const FNavigationToolViewModelPtr Item = WeakItem.Pin();
+	if (!Item.IsValid())
 	{
-		return !WeakView.Pin()->IsItemReadOnly(WeakItem.Pin());
+		return false;
 	}
-	return false;
+
+	return !View->IsItemReadOnly(Item);
 }
 
 FText SNavigationToolLabelItem::GetItemText() const
 {
-	const FNavigationToolItemPtr Item = WeakItem.Pin();
+	const FNavigationToolViewModelPtr Item = WeakItem.Pin();
 	if (!Item.IsValid())
 	{
 		return FText::GetEmpty();
@@ -143,7 +156,7 @@ bool SNavigationToolLabelItem::OnVerifyItemLabelChanged(const FText& InLabel, FT
 
 void SNavigationToolLabelItem::OnLabelTextCommitted(const FText& InLabel, ETextCommit::Type InCommitInfo)
 {
-	const FNavigationToolItemPtr Item = WeakItem.Pin();
+	const FNavigationToolViewModelPtr Item = WeakItem.Pin();
 	if (!Item.IsValid())
 	{
 		return;
@@ -167,22 +180,23 @@ void SNavigationToolLabelItem::OnLabelTextCommitted(const FText& InLabel, ETextC
 
 void SNavigationToolLabelItem::RenameItem(const FText& InLabel)
 {
-	const FNavigationToolItemPtr Item = WeakItem.Pin();
+	const FNavigationToolViewModelPtr Item = WeakItem.Pin();
 	if (!Item.IsValid())
 	{
 		return;
 	}
 
-	IRenameableExtension* const RenameableExtension = Item->CastTo<IRenameableExtension>();
-	if (!RenameableExtension)
+	const TViewModelPtr<IRenameableExtension> RenameableExtension = Item.ImplicitCast();
+	if (!RenameableExtension.IsValid())
 	{
 		return;
 	}
 
 	const FScopedTransaction Transaction(LOCTEXT("ItemRename", "Rename Item"));
 
-	const bool bRenameSuccessful = RenameableExtension->Rename(InLabel.ToString());
-	if (bRenameSuccessful && WeakView.IsValid())
+	RenameableExtension->Rename(InLabel);
+
+	if (WeakView.IsValid())
 	{
 		WeakView.Pin()->SetKeyboardFocus();
 	}
diff --git a/Engine/Plugins/Experimental/SequenceNavigator/Source/SequenceNavigator/Private/Widgets/Columns/SNavigationToolLabelItem.h b/Engine/Plugins/Experimental/SequenceNavigator/Source/SequenceNavigator/Private/Widgets/Columns/SNavigationToolLabelItem.h
--- a/Engine/Plugins/Experimental/SequenceNavigator/Source/SequenceNavigator/Private/Widgets/Columns/SNavigationToolLabelItem.h
+++ b/Engine/Plugins/Experimental/SequenceNavigator/Source/SequenceNavigator/Private/Widgets/Columns/SNavigationToolLabelItem.h
@@ -24,7 +24,7 @@ public:
 	SLATE_END_ARGS()
 
 	void Construct(const FArguments& InArgs
-		, const FNavigationToolItemRef& InItem
+		, const FNavigationToolViewModelPtr& InItem
 		, const TSharedRef<SNavigationToolTreeRow>& InRowWidget);
 
 	bool IsReadOnly() const;
@@ -46,7 +46,7 @@ public:
 	virtual const FInlineEditableTextBlockStyle* GetTextBlockStyle() const;
 
 protected:
-	FNavigationToolItemWeakPtr WeakItem;
+	FNavigationToolViewModelWeakPtr WeakItem;
 
 	TWeakPtr<INavigationToolView> WeakView;
 
diff --git a/Engine/Plugins/Experimental/SequenceNavigator/Source/SequenceNavigator/Private/Widgets/Columns/SNavigationToolLength.cpp b/Engine/Plugins/Experimental/SequenceNavigator/Source/SequenceNavigator/Private/Widgets/Columns/SNavigationToolLength.cpp
--- a/Engine/Plugins/Experimental/SequenceNavigator/Source/SequenceNavigator/Private/Widgets/Columns/SNavigationToolLength.cpp
+++ b/Engine/Plugins/Experimental/SequenceNavigator/Source/SequenceNavigator/Private/Widgets/Columns/SNavigationToolLength.cpp
@@ -3,7 +3,6 @@
 #include "SNavigationToolLength.h"
 #include "ISequencer.h"
 #include "Items/NavigationToolItemUtils.h"
-#include "Items/NavigationToolSequence.h"
 #include "NavigationToolView.h"
 #include "Sections/MovieSceneSubSection.h"
 #include "Widgets/SNavigationToolTreeRow.h"
@@ -13,8 +12,10 @@
 namespace UE::SequenceNavigator
 {
 
+using namespace Sequencer;
+
 void SNavigationToolLength::Construct(const FArguments& InArgs
-	, const FNavigationToolItemRef& InItem
+	, const FNavigationToolViewModelPtr& InItem
 	, const TSharedRef<INavigationToolView>& InView
 	, const TSharedRef<SNavigationToolTreeRow>& InRowWidget)
 {
@@ -22,18 +23,12 @@ void SNavigationToolLength::Construct(const FArguments& InArgs
 	WeakView = InView;
 	WeakRowWidget = InRowWidget;
 
-	const FNavigationToolSequence* const SequenceItem = InItem->CastTo<FNavigationToolSequence>();
-	if (!SequenceItem || !SequenceItem->GetSubSection())
-	{
-		return;
-	}
-
 	SNavigationToolTime::Construct(SNavigationToolTime::FArguments(), InItem, InView, InRowWidget);
 }
 
 double SNavigationToolLength::GetFrameTimeValue() const
 {
-	if (UMovieSceneSubSection* const SubSection = Cast<UMovieSceneSubSection>(GetSequenceItemSubSection(WeakItem.Pin())))
+	if (const UMovieSceneSubSection* const SubSection = ItemUtils::GetSequenceItemSubSection(WeakItem.Pin()))
 	{
 		return SubSection->GetRange().Size<FFrameNumber>().Value;
 	}
@@ -49,7 +44,7 @@ TSharedPtr<INumericTypeInterface<double>> SNavigationToolLength::GetNumericTypeI
 {
 	if (const TSharedPtr<ISequencer> Sequencer = GetSequencer())
 	{
-		return Sequencer->GetNumericTypeInterface(UE::Sequencer::ENumericIntent::Duration);
+		return Sequencer->GetNumericTypeInterface(ENumericIntent::Duration);
 	}
 	return nullptr;
 }
diff --git a/Engine/Plugins/Experimental/SequenceNavigator/Source/SequenceNavigator/Private/Widgets/Columns/SNavigationToolLength.h b/Engine/Plugins/Experimental/SequenceNavigator/Source/SequenceNavigator/Private/Widgets/Columns/SNavigationToolLength.h
--- a/Engine/Plugins/Experimental/SequenceNavigator/Source/SequenceNavigator/Private/Widgets/Columns/SNavigationToolLength.h
+++ b/Engine/Plugins/Experimental/SequenceNavigator/Source/SequenceNavigator/Private/Widgets/Columns/SNavigationToolLength.h
@@ -20,16 +20,16 @@ public:
 	virtual ~SNavigationToolLength() override {}
 
 	void Construct(const FArguments& InArgs
-		, const FNavigationToolItemRef& InItem
+		, const FNavigationToolViewModelPtr& InItem
 		, const TSharedRef<INavigationToolView>& InView
 		, const TSharedRef<SNavigationToolTreeRow>& InRowWidget);
 
-private:
+protected:
+	//~ Begin SNavigationToolTime
 	virtual double GetFrameTimeValue() const override;
-
 	virtual FText GetTransactionText() const override;
-
 	virtual TSharedPtr<INumericTypeInterface<double>> GetNumericTypeInterface() const override;
+	//~ End SNavigationToolTime
 };
 
 } // namespace UE::SequenceNavigator
diff --git a/Engine/Plugins/Experimental/SequenceNavigator/Source/SequenceNavigator/Private/Widgets/Columns/SNavigationToolLock.cpp b/Engine/Plugins/Experimental/SequenceNavigator/Source/SequenceNavigator/Private/Widgets/Columns/SNavigationToolLock.cpp
--- a/Engine/Plugins/Experimental/SequenceNavigator/Source/SequenceNavigator/Private/Widgets/Columns/SNavigationToolLock.cpp
+++ b/Engine/Plugins/Experimental/SequenceNavigator/Source/SequenceNavigator/Private/Widgets/Columns/SNavigationToolLock.cpp
@@ -12,6 +12,8 @@
 namespace UE::SequenceNavigator
 {
 
+using namespace Sequencer;
+
 class FLockDragDropOp : public FDragDropOperation
 {
 public:
@@ -40,37 +42,25 @@ public:
 	}
 };
 
-EItemSequenceLockState GetItemLockState(const FNavigationToolItemPtr& InItem)
+ELockableLockState GetItemLockState(const FNavigationToolViewModelPtr& InItem)
 {
-	if (!InItem.IsValid())
-	{
-		return EItemSequenceLockState::None;
-	}
-
-	ISequenceLockableExtension* const LockableItem = InItem->CastTo<ISequenceLockableExtension>();
-	if (!LockableItem)
+	if (const TViewModelPtr<ILockableExtension> LockableItem = InItem.ImplicitCast())
 	{
-		return EItemSequenceLockState::None;
+		return LockableItem->GetLockState();
 	}
-
-	return LockableItem->GetLockState();
+	return ELockableLockState::None;
 }
 
-void SetItemLocked(const FNavigationToolItemPtr& InItem, const bool bInIsLocked)
+void SetItemLocked(const FNavigationToolViewModelPtr& InItem, const bool bInIsLocked)
 {
-	if (!InItem.IsValid())
-	{
-		return;
-	}
-
-	if (ISequenceLockableExtension* const LockableItem = InItem->CastTo<ISequenceLockableExtension>())
+	if (const TViewModelPtr<ILockableExtension> LockableItem = InItem.ImplicitCast())
 	{
 		LockableItem->SetIsLocked(bInIsLocked);
 	}
 }
 
 void SNavigationToolLock::Construct(const FArguments& InArgs
-	, const FNavigationToolItemRef& InItem
+	, const FNavigationToolViewModelPtr& InItem
 	, const TSharedRef<INavigationToolView>& InView
 	, const TSharedRef<SNavigationToolTreeRow>& InRowWidget)
 {
@@ -85,20 +75,20 @@ void SNavigationToolLock::Construct(const FArguments& InArgs
 
 FSlateColor SNavigationToolLock::GetForegroundColor() const
 {
-	const FNavigationToolItemPtr Item = WeakItem.Pin();
+	const FNavigationToolViewModelPtr Item = WeakItem.Pin();
 	if (!Item.IsValid())
 	{
 		return FSlateColor::UseForeground();
 	}
 
-	const bool bIsItemSelected = WeakView.IsValid() && WeakView.Pin()->IsItemSelected(Item.ToSharedRef());
+	const bool bIsItemSelected = WeakView.IsValid() && WeakView.Pin()->IsItemSelected(Item);
 	const bool bIsItemHovered = WeakRowWidget.IsValid() && WeakRowWidget.Pin()->IsHovered();
 	const bool bAlwaysShowLock = GetDefault<UNavigationToolSettings>()->ShouldAlwaysShowLockState();
 
 	// We can hide the brush if settings for "Always Showing Lock State" is OFF
 	// and item is not locked while also not being selected or hovered
 	if (!bAlwaysShowLock
-		&& GetLockState() == EItemSequenceLockState::None
+		&& GetLockState() == ELockableLockState::None
 		&& !bIsItemSelected
 		&& !bIsItemHovered)
 	{
@@ -109,13 +99,13 @@ FSlateColor SNavigationToolLock::GetForegroundColor() const
 	{
 		switch (GetLockState())
 		{
-		case EItemSequenceLockState::None:
+		case ELockableLockState::None:
 			return FStyleColors::White25;
 
-		case EItemSequenceLockState::PartiallyLocked:
+		case ELockableLockState::PartiallyLocked:
 			return FStyleColors::ForegroundHover;
 
-		case EItemSequenceLockState::Locked:
+		case ELockableLockState::Locked:
 			return FStyleColors::ForegroundHover;
 		}
 	}
@@ -123,13 +113,13 @@ FSlateColor SNavigationToolLock::GetForegroundColor() const
 	{
 		switch (GetLockState())
 		{
-		case EItemSequenceLockState::None:
+		case ELockableLockState::None:
 			return FStyleColors::Transparent;
 
-		case EItemSequenceLockState::PartiallyLocked:
+		case ELockableLockState::PartiallyLocked:
 			return FStyleColors::White25;
 
-		case EItemSequenceLockState::Locked:
+		case ELockableLockState::Locked:
 			return FStyleColors::Foreground;
 		}
 	}
@@ -139,24 +129,24 @@ FSlateColor SNavigationToolLock::GetForegroundColor() const
 
 const FSlateBrush* SNavigationToolLock::GetBrush() const
 {
-	return (GetLockState() == EItemSequenceLockState::None)
+	return (GetLockState() == ELockableLockState::None)
 		? FAppStyle::GetBrush(TEXT("Icons.Unlock"))
 		: FAppStyle::GetBrush(TEXT("Icons.Lock"));
 }
 
-FReply SNavigationToolLock::OnDragDetected(const FGeometry& MyGeometry, const FPointerEvent& MouseEvent)
+FReply SNavigationToolLock::OnDragDetected(const FGeometry& InGeometry, const FPointerEvent& InPointerEvent)
 {
-	if (MouseEvent.IsMouseButtonDown(EKeys::LeftMouseButton))
+	if (InPointerEvent.IsMouseButtonDown(EKeys::LeftMouseButton))
 	{
-		const bool bShouldLock = GetLockState() == EItemSequenceLockState::Locked;
+		const bool bShouldLock = GetLockState() == ELockableLockState::Locked;
 		return FReply::Handled().BeginDragDrop(FLockDragDropOp::New(bShouldLock, UndoTransaction));
 	}
 	return FReply::Unhandled();
 }
 
-void SNavigationToolLock::OnDragEnter(const FGeometry& MyGeometry, const FDragDropEvent& DragDropEvent)
+void SNavigationToolLock::OnDragEnter(const FGeometry& InGeometry, const FDragDropEvent& InDragDropEvent)
 {
-	if (const TSharedPtr<FLockDragDropOp> DragDropOp = DragDropEvent.GetOperationAs<FLockDragDropOp>())
+	if (const TSharedPtr<FLockDragDropOp> DragDropOp = InDragDropEvent.GetOperationAs<FLockDragDropOp>())
 	{
 		SetIsLocked(DragDropOp->bShouldLock);
 	}
@@ -165,23 +155,27 @@ void SNavigationToolLock::OnDragEnter(const FGeometry& MyGeometry, const FDragDr
 FReply SNavigationToolLock::HandleClick()
 {
 	const TSharedPtr<INavigationToolView> ToolView = WeakView.Pin();
-	const FNavigationToolItemPtr Item = WeakItem.Pin();
+	if (!ToolView.IsValid())
+	{
+		return FReply::Unhandled();
+	}
 
-	if (!ToolView.IsValid() || !Item.IsValid())
+	const FNavigationToolViewModelPtr Item = WeakItem.Pin();
+	if (!Item.IsValid())
 	{
 		return FReply::Unhandled();
 	}
 
 	UndoTransaction = MakeUnique<FScopedTransaction>(LOCTEXT("SetNavigationToolItemLock", "Set Item Lock"));
 
-	const bool bNewIsLocked = (GetLockState() != EItemSequenceLockState::Locked);
+	const bool bNewIsLocked = (GetLockState() != ELockableLockState::Locked);
 
 	// We operate on all the selected items if the specified item is selected
-	if (ToolView->IsItemSelected(Item.ToSharedRef()))
+	if (ToolView->IsItemSelected(Item))
 	{
-		for (FNavigationToolItemPtr& SelectedItem : ToolView->GetSelectedItems())
+		for (FNavigationToolViewModelWeakPtr& WeakSelectedItem : ToolView->GetSelectedItems())
 		{
-			SetItemLocked(SelectedItem, bNewIsLocked);
+			SetItemLocked(WeakSelectedItem.Pin(), bNewIsLocked);
 		}
 	}
 	else
@@ -192,23 +186,23 @@ FReply SNavigationToolLock::HandleClick()
 	return FReply::Handled().DetectDrag(SharedThis(this), EKeys::LeftMouseButton);
 }
 
-FReply SNavigationToolLock::OnMouseButtonDoubleClick(const FGeometry& InMyGeometry, const FPointerEvent& InMouseEvent)
+FReply SNavigationToolLock::OnMouseButtonDoubleClick(const FGeometry& InInGeometry, const FPointerEvent& InPointerEvent)
 {
 	return HandleClick();
 }
 
-FReply SNavigationToolLock::OnMouseButtonDown(const FGeometry& MyGeometry, const FPointerEvent& MouseEvent)
+FReply SNavigationToolLock::OnMouseButtonDown(const FGeometry& InGeometry, const FPointerEvent& InPointerEvent)
 {
-	if (MouseEvent.GetEffectingButton() != EKeys::LeftMouseButton)
+	if (InPointerEvent.GetEffectingButton() != EKeys::LeftMouseButton)
 	{
 		return FReply::Unhandled();
 	}
 	return HandleClick();
 }
 
-FReply SNavigationToolLock::OnMouseButtonUp(const FGeometry& MyGeometry, const FPointerEvent& MouseEvent)
+FReply SNavigationToolLock::OnMouseButtonUp(const FGeometry& InGeometry, const FPointerEvent& InPointerEvent)
 {
-	if (MouseEvent.GetEffectingButton() == EKeys::LeftMouseButton)
+	if (InPointerEvent.GetEffectingButton() == EKeys::LeftMouseButton)
 	{
 		UndoTransaction.Reset();
 
@@ -223,7 +217,7 @@ void SNavigationToolLock::OnMouseCaptureLost(const FCaptureLostEvent& CaptureLos
 	UndoTransaction.Reset();
 }
 
-EItemSequenceLockState SNavigationToolLock::GetLockState() const
+ELockableLockState SNavigationToolLock::GetLockState() const
 {
 	return GetItemLockState(WeakItem.Pin());
 }
diff --git a/Engine/Plugins/Experimental/SequenceNavigator/Source/SequenceNavigator/Private/Widgets/Columns/SNavigationToolLock.h b/Engine/Plugins/Experimental/SequenceNavigator/Source/SequenceNavigator/Private/Widgets/Columns/SNavigationToolLock.h
--- a/Engine/Plugins/Experimental/SequenceNavigator/Source/SequenceNavigator/Private/Widgets/Columns/SNavigationToolLock.h
+++ b/Engine/Plugins/Experimental/SequenceNavigator/Source/SequenceNavigator/Private/Widgets/Columns/SNavigationToolLock.h
@@ -2,7 +2,7 @@
 
 #pragma once
 
-#include "Extensions/ISequenceLockableExtension.h"
+#include "MVVM/Extensions/ILockableExtension.h"
 #include "NavigationToolDefines.h"
 #include "ScopedTransaction.h"
 #include "Templates/UniquePtr.h"
@@ -31,37 +31,37 @@ public:
 	virtual ~SNavigationToolLock() override {}
 
 	void Construct(const FArguments& InArgs
-		, const FNavigationToolItemRef& InItem
+		, const FNavigationToolViewModelPtr& InItem
 		, const TSharedRef<INavigationToolView>& InView
 		, const TSharedRef<SNavigationToolTreeRow>& InRowWidget);
 
 	virtual FSlateColor GetForegroundColor() const override;
 
-	EItemSequenceLockState GetLockState() const;
+	Sequencer::ELockableLockState GetLockState() const;
 
 	void SetIsLocked(const bool bInIsLocked);
 
 	const FSlateBrush* GetBrush() const;
 
-	virtual FReply OnDragDetected(const FGeometry& MyGeometry, const FPointerEvent& MouseEvent) override;
+	virtual FReply OnDragDetected(const FGeometry& InGeometry, const FPointerEvent& InPointerEvent) override;
 
 	/** If a lock drag drop operation has entered this widget, set its item to the new lock state */
-	virtual void OnDragEnter(const FGeometry& MyGeometry, const FDragDropEvent& DragDropEvent) override;
+	virtual void OnDragEnter(const FGeometry& InGeometry, const FDragDropEvent& InDragDropEvent) override;
 
 	FReply HandleClick();
 
-	virtual FReply OnMouseButtonDoubleClick(const FGeometry& InMyGeometry, const FPointerEvent& InMouseEvent) override;
+	virtual FReply OnMouseButtonDoubleClick(const FGeometry& InInGeometry, const FPointerEvent& InPointerEvent) override;
 
-	virtual FReply OnMouseButtonDown(const FGeometry& MyGeometry, const FPointerEvent& MouseEvent) override;
+	virtual FReply OnMouseButtonDown(const FGeometry& InGeometry, const FPointerEvent& InPointerEvent) override;
 
-	virtual FReply OnMouseButtonUp(const FGeometry& MyGeometry, const FPointerEvent& MouseEvent) override;
+	virtual FReply OnMouseButtonUp(const FGeometry& InGeometry, const FPointerEvent& InPointerEvent) override;
 
 	virtual void OnMouseCaptureLost(const FCaptureLostEvent& CaptureLostEvent) override;
 
-	//void OnSetItemLockedState(const FNavigationToolItemPtr& InItem, const bool bInIsLocked);
+	//void OnSetItemLockedState(const FNavigationToolViewModelPtr& InItem, const bool bInIsLocked);
 
-private:
-	FNavigationToolItemWeakPtr WeakItem;
+protected:
+	FNavigationToolViewModelWeakPtr WeakItem;
 
 	TWeakPtr<INavigationToolView> WeakView;
 
diff --git a/Engine/Plugins/Experimental/SequenceNavigator/Source/SequenceNavigator/Private/Widgets/Columns/SNavigationToolMarkerVisibility.cpp b/Engine/Plugins/Experimental/SequenceNavigator/Source/SequenceNavigator/Private/Widgets/Columns/SNavigationToolMarkerVisibility.cpp
--- a/Engine/Plugins/Experimental/SequenceNavigator/Source/SequenceNavigator/Private/Widgets/Columns/SNavigationToolMarkerVisibility.cpp
+++ b/Engine/Plugins/Experimental/SequenceNavigator/Source/SequenceNavigator/Private/Widgets/Columns/SNavigationToolMarkerVisibility.cpp
@@ -13,6 +13,8 @@
 namespace UE::SequenceNavigator
 {
 
+using namespace Sequencer;
+
 class FMarkerVisibilityDragDropOp : public FDragDropOperation
 {
 public:
@@ -41,14 +43,14 @@ public:
 	}
 };
 
-EItemMarkerVisibility GetItemMarkerVisibility(const FNavigationToolItemPtr& InItem)
+EItemMarkerVisibility GetItemMarkerVisibility(const FNavigationToolViewModelPtr& InItem)
 {
 	if (!InItem.IsValid())
 	{
 		return EItemMarkerVisibility::None;
 	}
 
-	IMarkerVisibilityExtension* const MarkerVisibilityItem = InItem->CastTo<IMarkerVisibilityExtension>();
+	const TViewModelPtr<IMarkerVisibilityExtension> MarkerVisibilityItem = InItem.ImplicitCast();
 	if (!MarkerVisibilityItem)
 	{
 		return EItemMarkerVisibility::None;
@@ -57,14 +59,14 @@ EItemMarkerVisibility GetItemMarkerVisibility(const FNavigationToolItemPtr& InIt
 	return MarkerVisibilityItem->GetMarkerVisibility();
 }
 
-void SetItemMarkerVisibility(const FNavigationToolItemPtr& InItem, const bool bInMarkersVisible)
+void SetItemMarkerVisibility(const FNavigationToolViewModelPtr& InItem, const bool bInMarkersVisible)
 {
 	if (!InItem.IsValid())
 	{
 		return;
 	}
 
-	if (IMarkerVisibilityExtension* const MarkerVisibilityItem = InItem->CastTo<IMarkerVisibilityExtension>())
+	if (const TViewModelPtr<IMarkerVisibilityExtension> MarkerVisibilityItem = InItem.ImplicitCast())
 	{
 		MarkerVisibilityItem->SetMarkerVisibility(bInMarkersVisible);
 	}
@@ -72,7 +74,7 @@ void SetItemMarkerVisibility(const FNavigationToolItemPtr& InItem, const bool bI
 
 void SNavigationToolMarkerVisibility::Construct(const FArguments& InArgs
 	, const TSharedRef<FNavigationToolMarkerVisibilityColumn>& InColumn
-	, const FNavigationToolItemRef& InItem
+	, const FNavigationToolViewModelPtr& InItem
 	, const TSharedRef<INavigationToolView>& InView
 	, const TSharedRef<SNavigationToolTreeRow>& InRowWidget)
 {
@@ -87,18 +89,18 @@ void SNavigationToolMarkerVisibility::Construct(const FArguments& InArgs
 		.Image(this, &SNavigationToolMarkerVisibility::GetBrush));
 }
 
-FReply SNavigationToolMarkerVisibility::OnDragDetected(const FGeometry& MyGeometry, const FPointerEvent& MouseEvent)
+FReply SNavigationToolMarkerVisibility::OnDragDetected(const FGeometry& InGeometry, const FPointerEvent& InPointerEvent)
 {
-	if (MouseEvent.IsMouseButtonDown(EKeys::LeftMouseButton))
+	if (InPointerEvent.IsMouseButtonDown(EKeys::LeftMouseButton))
 	{
 		return FReply::Handled().BeginDragDrop(FMarkerVisibilityDragDropOp::New(GetMarkerVisibility(), UndoTransaction));
 	}
 	return FReply::Unhandled();
 }
 
-void SNavigationToolMarkerVisibility::OnDragEnter(const FGeometry& MyGeometry, const FDragDropEvent& DragDropEvent)
+void SNavigationToolMarkerVisibility::OnDragEnter(const FGeometry& InGeometry, const FDragDropEvent& InDragDropEvent)
 {
-	const TSharedPtr<FMarkerVisibilityDragDropOp> DragDropOp = DragDropEvent.GetOperationAs<FMarkerVisibilityDragDropOp>();
+	const TSharedPtr<FMarkerVisibilityDragDropOp> DragDropOp = InDragDropEvent.GetOperationAs<FMarkerVisibilityDragDropOp>();
 	if (DragDropOp.IsValid())
 	{
 		const bool bNewVisible = (DragDropOp->MarkerVisibility != EItemMarkerVisibility::Visible);
@@ -108,16 +110,25 @@ void SNavigationToolMarkerVisibility::OnDragEnter(const FGeometry& MyGeometry, c
 
 FReply SNavigationToolMarkerVisibility::HandleClick()
 {
-	if (!IsVisibilityWidgetEnabled())
+	const TSharedPtr<INavigationToolView> ToolView = WeakView.Pin();
+	if (!ToolView.IsValid())
+	{
+		return FReply::Unhandled();
+	}
+
+	const FNavigationToolViewModelPtr Item = WeakItem.Pin();
+	if (!Item.IsValid())
 	{
 		return FReply::Unhandled();
 	}
 
-	const TSharedPtr<INavigationToolView> ToolView = WeakView.Pin();
-	const FNavigationToolItemPtr Item = WeakItem.Pin();
 	const TSharedPtr<FNavigationToolMarkerVisibilityColumn> Column = WeakColumn.Pin();
+	if (!Column.IsValid())
+	{
+		return FReply::Unhandled();
+	}
 
-	if (!ToolView.IsValid() || !Item.IsValid() || !Column.IsValid())
+	if (!IsVisibilityWidgetEnabled())
 	{
 		return FReply::Unhandled();
 	}
@@ -128,11 +139,11 @@ FReply SNavigationToolMarkerVisibility::HandleClick()
 	const bool bNewVisible = (GetMarkerVisibility() != EItemMarkerVisibility::Visible);
 
 	// We operate on all the selected items if the specified item is selected
-	if (ToolView->IsItemSelected(Item.ToSharedRef()))
+	if (ToolView->IsItemSelected(Item))
 	{
-		for (const FNavigationToolItemPtr& SelectedItem : ToolView->GetSelectedItems())
+		for (const FNavigationToolViewModelWeakPtr& WeakSelectedItem : ToolView->GetSelectedItems())
 		{
-			if (IMarkerVisibilityExtension* const MarkerVisibilityItem = SelectedItem->CastTo<IMarkerVisibilityExtension>())
+			if (const TViewModelPtr<IMarkerVisibilityExtension> MarkerVisibilityItem = WeakSelectedItem.ImplicitPin())
 			{
 				MarkerVisibilityItem->SetMarkerVisibility(bNewVisible);
 			}
@@ -146,14 +157,14 @@ FReply SNavigationToolMarkerVisibility::HandleClick()
 	return FReply::Handled().DetectDrag(SharedThis(this), EKeys::LeftMouseButton);
 }
 
-FReply SNavigationToolMarkerVisibility::OnMouseButtonDoubleClick(const FGeometry& InMyGeometry, const FPointerEvent& InMouseEvent)
+FReply SNavigationToolMarkerVisibility::OnMouseButtonDoubleClick(const FGeometry& InInGeometry, const FPointerEvent& InPointerEvent)
 {
 	return HandleClick();
 }
 
-FReply SNavigationToolMarkerVisibility::OnMouseButtonDown(const FGeometry& MyGeometry, const FPointerEvent& MouseEvent)
+FReply SNavigationToolMarkerVisibility::OnMouseButtonDown(const FGeometry& InGeometry, const FPointerEvent& InPointerEvent)
 {
-	if (MouseEvent.GetEffectingButton() != EKeys::LeftMouseButton)
+	if (InPointerEvent.GetEffectingButton() != EKeys::LeftMouseButton)
 	{
 		return FReply::Unhandled();
 	}
@@ -161,9 +172,9 @@ FReply SNavigationToolMarkerVisibility::OnMouseButtonDown(const FGeometry& MyGeo
 	return HandleClick();
 }
 
-FReply SNavigationToolMarkerVisibility::OnMouseButtonUp(const FGeometry& MyGeometry, const FPointerEvent& MouseEvent)
+FReply SNavigationToolMarkerVisibility::OnMouseButtonUp(const FGeometry& InGeometry, const FPointerEvent& InPointerEvent)
 {
-	if (MouseEvent.GetEffectingButton() == EKeys::LeftMouseButton)
+	if (InPointerEvent.GetEffectingButton() == EKeys::LeftMouseButton)
 	{
 		UndoTransaction.Reset();
 		return FReply::Handled();
@@ -172,7 +183,7 @@ FReply SNavigationToolMarkerVisibility::OnMouseButtonUp(const FGeometry& MyGeome
 	return FReply::Unhandled();
 }
 
-void SNavigationToolMarkerVisibility::OnMouseCaptureLost(const FCaptureLostEvent& CaptureLostEvent)
+void SNavigationToolMarkerVisibility::OnMouseCaptureLost(const FCaptureLostEvent& InCaptureLostEvent)
 {
 	UndoTransaction.Reset();
 }
@@ -184,8 +195,8 @@ const FSlateBrush* SNavigationToolMarkerVisibility::GetBrush() const
 
 FSlateColor SNavigationToolMarkerVisibility::GetForegroundColor() const
 {
-	const IMarkerVisibilityExtension* const MarkerVisibilityExt = WeakItem.Pin()->CastTo<IMarkerVisibilityExtension>();
-	if (!MarkerVisibilityExt)
+	const TViewModelPtr<IMarkerVisibilityExtension> Item = WeakItem.ImplicitPin();
+	if (!Item.IsValid())
 	{
 		return FLinearColor::Transparent;
 	}
diff --git a/Engine/Plugins/Experimental/SequenceNavigator/Source/SequenceNavigator/Private/Widgets/Columns/SNavigationToolMarkerVisibility.h b/Engine/Plugins/Experimental/SequenceNavigator/Source/SequenceNavigator/Private/Widgets/Columns/SNavigationToolMarkerVisibility.h
--- a/Engine/Plugins/Experimental/SequenceNavigator/Source/SequenceNavigator/Private/Widgets/Columns/SNavigationToolMarkerVisibility.h
+++ b/Engine/Plugins/Experimental/SequenceNavigator/Source/SequenceNavigator/Private/Widgets/Columns/SNavigationToolMarkerVisibility.h
@@ -24,25 +24,25 @@ public:
 
 	void Construct(const FArguments& InArgs
 		, const TSharedRef<FNavigationToolMarkerVisibilityColumn>& InColumn
-		, const FNavigationToolItemRef& InItem
+		, const FNavigationToolViewModelPtr& InItem
 		, const TSharedRef<INavigationToolView>& InView
 		, const TSharedRef<SNavigationToolTreeRow>& InRowWidget);
 
-private:
+protected:
 	/** Returns whether the widget is enabled or not */
 	virtual bool IsVisibilityWidgetEnabled() const { return true; }
 
 	virtual const FSlateBrush* GetBrush() const;
 
-	virtual FReply OnDragDetected(const FGeometry& MyGeometry, const FPointerEvent& MouseEvent) override;
-	virtual void OnDragEnter(const FGeometry& MyGeometry, const FDragDropEvent& DragDropEvent) override;
+	virtual FReply OnDragDetected(const FGeometry& InGeometry, const FPointerEvent& InPointerEvent) override;
+	virtual void OnDragEnter(const FGeometry& InGeometry, const FDragDropEvent& InDragDropEvent) override;
 
 	FReply HandleClick();
 
-	virtual FReply OnMouseButtonDoubleClick(const FGeometry& InMyGeometry, const FPointerEvent& InMouseEvent) override;
-	virtual FReply OnMouseButtonDown(const FGeometry& MyGeometry, const FPointerEvent& MouseEvent) override;
-	virtual FReply OnMouseButtonUp(const FGeometry& MyGeometry, const FPointerEvent& MouseEvent) override;
-	virtual void OnMouseCaptureLost(const FCaptureLostEvent& CaptureLostEvent) override;
+	virtual FReply OnMouseButtonDoubleClick(const FGeometry& InInGeometry, const FPointerEvent& InPointerEvent) override;
+	virtual FReply OnMouseButtonDown(const FGeometry& InGeometry, const FPointerEvent& InPointerEvent) override;
+	virtual FReply OnMouseButtonUp(const FGeometry& InGeometry, const FPointerEvent& InPointerEvent) override;
+	virtual void OnMouseCaptureLost(const FCaptureLostEvent& InCaptureLostEvent) override;
 
 	virtual FSlateColor GetForegroundColor() const override;
 	
@@ -52,7 +52,7 @@ private:
 
 	TWeakPtr<FNavigationToolMarkerVisibilityColumn> WeakColumn;
 
-	FNavigationToolItemWeakPtr WeakItem;
+	FNavigationToolViewModelWeakPtr WeakItem;
 
 	TWeakPtr<INavigationToolView> WeakView;
 
diff --git a/Engine/Plugins/Experimental/SequenceNavigator/Source/SequenceNavigator/Private/Widgets/Columns/SNavigationToolOutTime.cpp b/Engine/Plugins/Experimental/SequenceNavigator/Source/SequenceNavigator/Private/Widgets/Columns/SNavigationToolOutTime.cpp
--- a/Engine/Plugins/Experimental/SequenceNavigator/Source/SequenceNavigator/Private/Widgets/Columns/SNavigationToolOutTime.cpp
+++ b/Engine/Plugins/Experimental/SequenceNavigator/Source/SequenceNavigator/Private/Widgets/Columns/SNavigationToolOutTime.cpp
@@ -12,8 +12,10 @@
 namespace UE::SequenceNavigator
 {
 
+using namespace Sequencer;
+
 void SNavigationToolOutTime::Construct(const FArguments& InArgs
-	, const FNavigationToolItemRef& InItem
+	, const FNavigationToolViewModelPtr& InItem
 	, const TSharedRef<INavigationToolView>& InView
 	, const TSharedRef<SNavigationToolTreeRow>& InRowWidget)
 {
@@ -21,19 +23,9 @@ void SNavigationToolOutTime::Construct(const FArguments& InArgs
 	WeakView = InView;
 	WeakRowWidget = InRowWidget;
 
-	if (!InItem->IsA<IOutTimeExtension>())
-	{
-		return;
-	}
-
 	SNavigationToolTime::Construct(SNavigationToolTime::FArguments(), InItem, InView, InRowWidget);
 }
 
-IOutTimeExtension* SNavigationToolOutTime::GetOutTimeExtension() const
-{
-	return WeakItem.IsValid() ? WeakItem.Pin()->CastTo<IOutTimeExtension>() : nullptr;
-}
-
 FName SNavigationToolOutTime::GetStyleName() const
 {
 	return TEXT("SpinBox.OutTime");
@@ -41,17 +33,17 @@ FName SNavigationToolOutTime::GetStyleName() const
 
 double SNavigationToolOutTime::GetFrameTimeValue() const
 {
-	if (const IOutTimeExtension* const OutTimeExtension = GetOutTimeExtension())
+	if (const TViewModelPtr<IOutTimeExtension> Item = WeakItem.ImplicitPin())
 	{
-		return OutTimeExtension->GetOutTime().Value;
+		return Item->GetOutTime().Value;
 	}
-	return 0;
+	return 0.;
 }
 
 void SNavigationToolOutTime::OnFrameTimeValueCommitted(const double InNewValue, const ETextCommit::Type InCommitType)
 {
-	IOutTimeExtension* const OutTimeExtension = GetOutTimeExtension();
-	if (!OutTimeExtension)
+	const TViewModelPtr<IOutTimeExtension> Item = WeakItem.ImplicitPin();
+	if (!Item.IsValid())
 	{
 		return;
 	}
@@ -59,7 +51,7 @@ void SNavigationToolOutTime::OnFrameTimeValueCommitted(const double InNewValue,
 	const bool bShouldTransact = !UndoTransaction.IsValid() && (InCommitType == ETextCommit::OnEnter);
 	const FScopedTransaction Transaction(GetTransactionText(), bShouldTransact);
 
-	return OutTimeExtension->SetOutTime(static_cast<int32>(InNewValue));
+	return Item->SetOutTime(static_cast<int32>(InNewValue));
 }
 
 FText SNavigationToolOutTime::GetTransactionText() const
diff --git a/Engine/Plugins/Experimental/SequenceNavigator/Source/SequenceNavigator/Private/Widgets/Columns/SNavigationToolOutTime.h b/Engine/Plugins/Experimental/SequenceNavigator/Source/SequenceNavigator/Private/Widgets/Columns/SNavigationToolOutTime.h
--- a/Engine/Plugins/Experimental/SequenceNavigator/Source/SequenceNavigator/Private/Widgets/Columns/SNavigationToolOutTime.h
+++ b/Engine/Plugins/Experimental/SequenceNavigator/Source/SequenceNavigator/Private/Widgets/Columns/SNavigationToolOutTime.h
@@ -8,8 +8,9 @@
 namespace UE::SequenceNavigator
 {
 
-class IOutTimeExtension;
+class FNavigationToolSequence;
 class INavigationToolView;
+class IOutTimeExtension;
 class SNavigationToolTreeRow;
 
 class SNavigationToolOutTime : public SNavigationToolTime
@@ -21,20 +22,17 @@ public:
 	virtual ~SNavigationToolOutTime() override {}
 
 	void Construct(const FArguments& InArgs
-		, const FNavigationToolItemRef& InItem
+		, const FNavigationToolViewModelPtr& InItem
 		, const TSharedRef<INavigationToolView>& InView
 		, const TSharedRef<SNavigationToolTreeRow>& InRowWidget);
 
 protected:
-	IOutTimeExtension* GetOutTimeExtension() const;
-
+	//~ Begin SNavigationToolTime
 	virtual FName GetStyleName() const override;
-
 	virtual double GetFrameTimeValue() const override;
-
 	virtual void OnFrameTimeValueCommitted(const double InNewValue, const ETextCommit::Type InCommitType) override;
-
 	virtual FText GetTransactionText() const override;
+	//~ End SNavigationToolTime
 };
 
 } // namespace UE::SequenceNavigator
diff --git a/Engine/Plugins/Experimental/SequenceNavigator/Source/SequenceNavigator/Private/Widgets/Columns/SNavigationToolPlayhead.cpp b/Engine/Plugins/Experimental/SequenceNavigator/Source/SequenceNavigator/Private/Widgets/Columns/SNavigationToolPlayhead.cpp
--- a/Engine/Plugins/Experimental/SequenceNavigator/Source/SequenceNavigator/Private/Widgets/Columns/SNavigationToolPlayhead.cpp
+++ b/Engine/Plugins/Experimental/SequenceNavigator/Source/SequenceNavigator/Private/Widgets/Columns/SNavigationToolPlayhead.cpp
@@ -14,28 +14,15 @@
 namespace UE::SequenceNavigator
 {
 
-EItemContainsPlayhead GetItemContainsPlayhead(const FNavigationToolItemPtr& InItem)
-{
-	if (!InItem.IsValid())
-	{
-		return EItemContainsPlayhead::None;
-	}
-
-	IPlayheadExtension* const PlayheadItem = InItem->CastTo<IPlayheadExtension>();
-	if (!PlayheadItem)
-	{
-		return EItemContainsPlayhead::None;
-	}
-
-	return PlayheadItem->ContainsPlayhead();
-}
+using namespace Sequencer;
 
 void SNavigationToolPlayhead::Construct(const FArguments& InArgs
 	, const TSharedRef<FNavigationToolPlayheadColumn>& InColumn
-	, const FNavigationToolItemRef& InItem
+	, const FNavigationToolViewModelPtr& InItem
 	, const TSharedRef<INavigationToolView>& InView
 	, const TSharedRef<SNavigationToolTreeRow>& InRowWidget)
 {
+	check(InItem.IsValid());
 	WeakColumn = InColumn;
 	WeakItem = InItem;
 	WeakView = InView;
@@ -90,7 +77,11 @@ FSlateColor SNavigationToolPlayhead::GetForegroundColor() const
 
 EItemContainsPlayhead SNavigationToolPlayhead::GetContainsPlayhead() const
 {
-	return GetItemContainsPlayhead(WeakItem.Pin());
+	if (const TViewModelPtr<IPlayheadExtension> PlayheadItem = WeakItem.ImplicitPin())
+	{
+		return PlayheadItem->ContainsPlayhead();
+	}
+	return EItemContainsPlayhead::None;
 }
 
 } // namespace UE::SequenceNavigator
diff --git a/Engine/Plugins/Experimental/SequenceNavigator/Source/SequenceNavigator/Private/Widgets/Columns/SNavigationToolPlayhead.h b/Engine/Plugins/Experimental/SequenceNavigator/Source/SequenceNavigator/Private/Widgets/Columns/SNavigationToolPlayhead.h
--- a/Engine/Plugins/Experimental/SequenceNavigator/Source/SequenceNavigator/Private/Widgets/Columns/SNavigationToolPlayhead.h
+++ b/Engine/Plugins/Experimental/SequenceNavigator/Source/SequenceNavigator/Private/Widgets/Columns/SNavigationToolPlayhead.h
@@ -22,11 +22,11 @@ public:
 
 	void Construct(const FArguments& InArgs
 		, const TSharedRef<FNavigationToolPlayheadColumn>& InColumn
-		, const FNavigationToolItemRef& InItem
+		, const FNavigationToolViewModelPtr& InItem
 		, const TSharedRef<INavigationToolView>& InView
 		, const TSharedRef<SNavigationToolTreeRow>& InRowWidget);
 
-private:
+protected:
 	/** Returns whether the widget is enabled or not */
 	virtual bool IsVisibilityWidgetEnabled() const { return true; }
 
@@ -38,7 +38,7 @@ private:
 
 	TWeakPtr<FNavigationToolPlayheadColumn> WeakColumn;
 
-	FNavigationToolItemWeakPtr WeakItem;
+	FNavigationToolViewModelWeakPtr WeakItem;
 
 	TWeakPtr<INavigationToolView> WeakView;
 
diff --git a/Engine/Plugins/Experimental/SequenceNavigator/Source/SequenceNavigator/Private/Widgets/Columns/SNavigationToolRevisionControl.cpp b/Engine/Plugins/Experimental/SequenceNavigator/Source/SequenceNavigator/Private/Widgets/Columns/SNavigationToolRevisionControl.cpp
--- a/Engine/Plugins/Experimental/SequenceNavigator/Source/SequenceNavigator/Private/Widgets/Columns/SNavigationToolRevisionControl.cpp
+++ b/Engine/Plugins/Experimental/SequenceNavigator/Source/SequenceNavigator/Private/Widgets/Columns/SNavigationToolRevisionControl.cpp
@@ -13,57 +13,11 @@
 namespace UE::SequenceNavigator
 {
 
-EItemRevisionControlState GetRevisionControlState(const FNavigationToolItemPtr& InItem)
-{
-	if (!InItem.IsValid())
-	{
-		return EItemRevisionControlState::None;
-	}
-
-	const IRevisionControlExtension* const RevisionControlItem = InItem->CastTo<IRevisionControlExtension>();
-	if (!RevisionControlItem)
-	{
-		return EItemRevisionControlState::None;
-	}
-
-	return RevisionControlItem->GetRevisionControlState();
-}
-
-const FSlateBrush* GetRevisionControlStatusIcon(const FNavigationToolItemPtr& InItem)
-{
-	if (!InItem.IsValid())
-	{
-		return nullptr;
-	}
-
-	const IRevisionControlExtension* const RevisionControlItem = InItem->CastTo<IRevisionControlExtension>();
-	if (!RevisionControlItem)
-	{
-		return nullptr;
-	}
-
-	return RevisionControlItem->GetRevisionControlStatusIcon();
-}
-
-FText GetRevisionControlStatusText(const FNavigationToolItemPtr& InItem)
-{
-	if (!InItem.IsValid())
-	{
-		return FText::GetEmpty();
-	}
-
-	const IRevisionControlExtension* const RevisionControlItem = InItem->CastTo<IRevisionControlExtension>();
-	if (!RevisionControlItem)
-	{
-		return FText::GetEmpty();
-	}
-
-	return RevisionControlItem->GetRevisionControlStatusText();
-}
+using namespace Sequencer;
 
 void SNavigationToolRevisionControl::Construct(const FArguments& InArgs
 	, const TSharedRef<FNavigationToolRevisionControlColumn>& InColumn
-	, const FNavigationToolItemRef& InItem
+	, const FNavigationToolViewModelPtr& InItem
 	, const TSharedRef<INavigationToolView>& InView
 	, const TSharedRef<SNavigationToolTreeRow>& InRowWidget)
 {
@@ -81,7 +35,13 @@ void SNavigationToolRevisionControl::Construct(const FArguments& InArgs
 
 FSlateColor SNavigationToolRevisionControl::GetForegroundColor() const
 {
-	switch (GetRevisionControlState(WeakItem.Pin()))
+	const TViewModelPtr<IRevisionControlExtension> Item = WeakItem.ImplicitPin();
+	if (!Item.IsValid())
+	{
+		return FStyleColors::Transparent;
+	}
+
+	switch (Item->GetRevisionControlState())
 	{
 	case EItemRevisionControlState::None:
 		return FStyleColors::Transparent;
@@ -90,17 +50,26 @@ FSlateColor SNavigationToolRevisionControl::GetForegroundColor() const
 	case EItemRevisionControlState::SourceControlled:
 		return FStyleColors::Foreground;
 	}
+
 	return FStyleColors::Transparent;
 }
 
 const FSlateBrush* SNavigationToolRevisionControl::GetBrush() const
 {
-	return GetRevisionControlStatusIcon(WeakItem.Pin());
+	if (const TViewModelPtr<IRevisionControlExtension> Item = WeakItem.ImplicitPin())
+	{
+		return Item->GetRevisionControlStatusIcon();
+	}
+	return nullptr;
 }
 
 FText SNavigationToolRevisionControl::GetToolTipText() const
 {
-	return GetRevisionControlStatusText(WeakItem.Pin());
+	if (const TViewModelPtr<IRevisionControlExtension> RevisionControlItem = WeakItem.ImplicitPin())
+	{
+		return RevisionControlItem->GetRevisionControlStatusText();
+	}
+	return FText::GetEmpty();
 }
 
 } // namespace UE::SequenceNavigator
diff --git a/Engine/Plugins/Experimental/SequenceNavigator/Source/SequenceNavigator/Private/Widgets/Columns/SNavigationToolRevisionControl.h b/Engine/Plugins/Experimental/SequenceNavigator/Source/SequenceNavigator/Private/Widgets/Columns/SNavigationToolRevisionControl.h
--- a/Engine/Plugins/Experimental/SequenceNavigator/Source/SequenceNavigator/Private/Widgets/Columns/SNavigationToolRevisionControl.h
+++ b/Engine/Plugins/Experimental/SequenceNavigator/Source/SequenceNavigator/Private/Widgets/Columns/SNavigationToolRevisionControl.h
@@ -10,6 +10,7 @@ namespace UE::SequenceNavigator
 
 class FNavigationToolRevisionControlColumn;
 class INavigationToolView;
+class IRevisionControlExtension;
 class SNavigationToolTreeRow;
 
 class SNavigationToolRevisionControl : public SImage
@@ -20,11 +21,11 @@ public:
 
 	void Construct(const FArguments& InArgs
 		, const TSharedRef<FNavigationToolRevisionControlColumn>& InColumn
-		, const FNavigationToolItemRef& InItem
+		, const FNavigationToolViewModelPtr& InItem
 		, const TSharedRef<INavigationToolView>& InView
 		, const TSharedRef<SNavigationToolTreeRow>& InRowWidget);
 
-private:
+protected:
 	virtual FSlateColor GetForegroundColor() const override;
 	
 	virtual const FSlateBrush* GetBrush() const;
@@ -32,7 +33,7 @@ private:
 
 	TWeakPtr<FNavigationToolRevisionControlColumn> WeakColumn;
 
-	FNavigationToolItemWeakPtr WeakItem;
+	FNavigationToolViewModelWeakPtr WeakItem;
 
 	TWeakPtr<INavigationToolView> WeakView;
 
diff --git a/Engine/Plugins/Experimental/SequenceNavigator/Source/SequenceNavigator/Private/Widgets/Columns/SNavigationToolStartFrameOffset.cpp b/Engine/Plugins/Experimental/SequenceNavigator/Source/SequenceNavigator/Private/Widgets/Columns/SNavigationToolStartFrameOffset.cpp
--- a/Engine/Plugins/Experimental/SequenceNavigator/Source/SequenceNavigator/Private/Widgets/Columns/SNavigationToolStartFrameOffset.cpp
+++ b/Engine/Plugins/Experimental/SequenceNavigator/Source/SequenceNavigator/Private/Widgets/Columns/SNavigationToolStartFrameOffset.cpp
@@ -11,8 +11,10 @@
 namespace UE::SequenceNavigator
 {
 
+using namespace Sequencer;
+
 void SNavigationToolStartFrameOffset::Construct(const FArguments& InArgs
-	, const FNavigationToolItemRef& InItem
+	, const FNavigationToolViewModelPtr& InItem
 	, const TSharedRef<INavigationToolView>& InView
 	, const TSharedRef<SNavigationToolTreeRow>& InRowWidget)
 {
@@ -20,18 +22,12 @@ void SNavigationToolStartFrameOffset::Construct(const FArguments& InArgs
 	WeakView = InView;
 	WeakRowWidget = InRowWidget;
 
-	const FNavigationToolSequence* const SequenceItem = InItem->CastTo<FNavigationToolSequence>();
-	if (!SequenceItem || !SequenceItem->GetSubSection())
-	{
-		return;
-	}
-
 	SNavigationToolTime::Construct(SNavigationToolTime::FArguments(), InItem, InView, InRowWidget);
 }
 
 double SNavigationToolStartFrameOffset::GetFrameTimeValue() const
 {
-	if (const UMovieSceneSubSection* const SubSection = Cast<UMovieSceneSubSection>(GetSequenceItemSubSection(WeakItem.Pin())))
+	if (const UMovieSceneSubSection* const SubSection = ItemUtils::GetSequenceItemSubSection(WeakItem.Pin()))
 	{
 		return SubSection->Parameters.StartFrameOffset.Value;
 	}
@@ -40,7 +36,9 @@ double SNavigationToolStartFrameOffset::GetFrameTimeValue() const
 
 void SNavigationToolStartFrameOffset::OnFrameTimeValueCommitted(const double InNewValue, const ETextCommit::Type InCommitType)
 {
-	UMovieSceneSubSection* const SubSection = Cast<UMovieSceneSubSection>(GetSequenceItemSubSection(WeakItem.Pin()));
+	using namespace ItemUtils;
+
+	UMovieSceneSubSection* const SubSection = GetSequenceItemSubSection(WeakItem.Pin());
 	if (!SubSection)
 	{
 		return;
diff --git a/Engine/Plugins/Experimental/SequenceNavigator/Source/SequenceNavigator/Private/Widgets/Columns/SNavigationToolStartFrameOffset.h b/Engine/Plugins/Experimental/SequenceNavigator/Source/SequenceNavigator/Private/Widgets/Columns/SNavigationToolStartFrameOffset.h
--- a/Engine/Plugins/Experimental/SequenceNavigator/Source/SequenceNavigator/Private/Widgets/Columns/SNavigationToolStartFrameOffset.h
+++ b/Engine/Plugins/Experimental/SequenceNavigator/Source/SequenceNavigator/Private/Widgets/Columns/SNavigationToolStartFrameOffset.h
@@ -17,19 +17,19 @@ public:
 	SLATE_BEGIN_ARGS(SNavigationToolStartFrameOffset) {}
 	SLATE_END_ARGS()
 
-	virtual ~SNavigationToolStartFrameOffset() override {}
+	virtual ~SNavigationToolStartFrameOffset() override = default;
 
 	void Construct(const FArguments& InArgs
-		, const FNavigationToolItemRef& InItem
+		, const FNavigationToolViewModelPtr& InItem
 		, const TSharedRef<INavigationToolView>& InView
 		, const TSharedRef<SNavigationToolTreeRow>& InRowWidget);
 
-private:
+protected:
+	//~ Begin SNavigationToolTime
 	virtual double GetFrameTimeValue() const override;
-
 	virtual void OnFrameTimeValueCommitted(const double InNewValue, const ETextCommit::Type InCommitType) override;
-
 	virtual FText GetTransactionText() const override;
+	//~ End SNavigationToolTime
 };
 
 } // namespace UE::SequenceNavigator
diff --git a/Engine/Plugins/Experimental/SequenceNavigator/Source/SequenceNavigator/Private/Widgets/Columns/SNavigationToolTake.cpp b/Engine/Plugins/Experimental/SequenceNavigator/Source/SequenceNavigator/Private/Widgets/Columns/SNavigationToolTake.cpp
--- a/Engine/Plugins/Experimental/SequenceNavigator/Source/SequenceNavigator/Private/Widgets/Columns/SNavigationToolTake.cpp
+++ b/Engine/Plugins/Experimental/SequenceNavigator/Source/SequenceNavigator/Private/Widgets/Columns/SNavigationToolTake.cpp
@@ -20,11 +20,11 @@ namespace UE::SequenceNavigator
 {
 
 void SNavigationToolTake::Construct(const FArguments& InArgs
-	, const FNavigationToolItemWeakPtr& InWeakItem
+	, const FNavigationToolViewModelPtr& InItem
 	, const TSharedRef<INavigationToolView>& InView
 	, const TSharedRef<SNavigationToolTreeRow>& InRowWidget)
 {
-	WeakItem = InWeakItem;
+	WeakItem = InItem;
 	WeakView = InView;
 	WeakRowWidget = InRowWidget;
 
@@ -100,14 +100,14 @@ void SNavigationToolTake::CacheTakes()
 {
 	CachedTakes.Reset();
 
-	const TSharedPtr<INavigationToolItem> Item = WeakItem.Pin();
+	const FNavigationToolViewModelPtr Item = WeakItem.Pin();
 	if (!Item.IsValid())
 	{
 		return;
 	}
 
-	FNavigationToolSequence* const SequenceItem = Item->CastTo<FNavigationToolSequence>();
-	if (!SequenceItem)
+	const Sequencer::TViewModelPtr<FNavigationToolSequence> SequenceItem = Item.ImplicitCast();
+	if (!SequenceItem.IsValid())
 	{
 		return;
 	}
@@ -180,14 +180,14 @@ void SNavigationToolTake::SetActiveTake(const TSharedPtr<FSequenceTakeEntry>& In
 		return;
 	}
 
-	const TSharedPtr<INavigationToolItem> Item = WeakItem.Pin();
+	const FNavigationToolViewModelPtr Item = WeakItem.Pin();
 	if (!Item.IsValid())
 	{
 		return;
 	}
 
-	FNavigationToolSequence* const SequenceItem = Item->CastTo<FNavigationToolSequence>();
-	if (!SequenceItem)
+	const Sequencer::TViewModelPtr<FNavigationToolSequence> SequenceItem = Item.ImplicitCast();
+	if (!SequenceItem.IsValid())
 	{
 		return;
 	}
diff --git a/Engine/Plugins/Experimental/SequenceNavigator/Source/SequenceNavigator/Private/Widgets/Columns/SNavigationToolTake.h b/Engine/Plugins/Experimental/SequenceNavigator/Source/SequenceNavigator/Private/Widgets/Columns/SNavigationToolTake.h
--- a/Engine/Plugins/Experimental/SequenceNavigator/Source/SequenceNavigator/Private/Widgets/Columns/SNavigationToolTake.h
+++ b/Engine/Plugins/Experimental/SequenceNavigator/Source/SequenceNavigator/Private/Widgets/Columns/SNavigationToolTake.h
@@ -2,15 +2,12 @@
 
 #pragma once
 
-#include "Columns/NavigationToolTakeColumn.h"
 #include "NavigationToolDefines.h"
 #include "ScopedTransaction.h"
 #include "SNavigationToolTakeEntry.h"
 #include "Templates/UniquePtr.h"
 #include "Widgets/SCompoundWidget.h"
 
-class UMovieSceneSequence;
-
 namespace UE::SequenceNavigator
 {
 
@@ -26,11 +23,13 @@ public:
 	SLATE_END_ARGS()
 
 	void Construct(const FArguments& InArgs
-		, const FNavigationToolItemWeakPtr& InWeakItem
+		, const FNavigationToolViewModelPtr& InItem
 		, const TSharedRef<INavigationToolView>& InView
 		, const TSharedRef<SNavigationToolTreeRow>& InRowWidget);
 
-private:
+	void SetActiveTake(const TSharedPtr<FSequenceTakeEntry>& InTakeInfo);
+
+protected:
 	FSlateColor GetBorderColor() const;
 
 	FReply OnTakeEntrySelected(const TSharedRef<FSequenceTakeEntry> InTakeInfo);
@@ -41,9 +40,8 @@ private:
 	void OnSelectionChanged(const TSharedPtr<FSequenceTakeEntry> InTakeInfo, const ESelectInfo::Type InSelectType);
 
 	void CacheTakes();
-	void SetActiveTake(const TSharedPtr<FSequenceTakeEntry>& InTakeInfo);
 
-	FNavigationToolItemWeakPtr WeakItem;
+	FNavigationToolViewModelWeakPtr WeakItem;
 
 	TWeakPtr<INavigationToolView> WeakView;
 
diff --git a/Engine/Plugins/Experimental/SequenceNavigator/Source/SequenceNavigator/Private/Widgets/Columns/SNavigationToolTime.cpp b/Engine/Plugins/Experimental/SequenceNavigator/Source/SequenceNavigator/Private/Widgets/Columns/SNavigationToolTime.cpp
--- a/Engine/Plugins/Experimental/SequenceNavigator/Source/SequenceNavigator/Private/Widgets/Columns/SNavigationToolTime.cpp
+++ b/Engine/Plugins/Experimental/SequenceNavigator/Source/SequenceNavigator/Private/Widgets/Columns/SNavigationToolTime.cpp
@@ -14,11 +14,14 @@
 namespace UE::SequenceNavigator
 {
 
+using namespace Sequencer;
+
 void SNavigationToolTime::Construct(const FArguments& InArgs
-	, const FNavigationToolItemRef& InItem
+	, const FNavigationToolViewModelPtr& InItem
 	, const TSharedRef<INavigationToolView>& InView
 	, const TSharedRef<SNavigationToolTreeRow>& InRowWidget)
 {
+	check(InItem.IsValid());
 	WeakItem = InItem;
 	WeakView = InView;
 	WeakRowWidget = InRowWidget;
@@ -82,7 +85,7 @@ TSharedPtr<INumericTypeInterface<double>> SNavigationToolTime::GetNumericTypeInt
 {
 	if (const TSharedPtr<ISequencer> Sequencer = GetSequencer())
 	{
-		return Sequencer->GetNumericTypeInterface(UE::Sequencer::ENumericIntent::Position);
+		return Sequencer->GetNumericTypeInterface(ENumericIntent::Position);
 	}
 	return nullptr;
 }
diff --git a/Engine/Plugins/Experimental/SequenceNavigator/Source/SequenceNavigator/Private/Widgets/Columns/SNavigationToolTime.h b/Engine/Plugins/Experimental/SequenceNavigator/Source/SequenceNavigator/Private/Widgets/Columns/SNavigationToolTime.h
--- a/Engine/Plugins/Experimental/SequenceNavigator/Source/SequenceNavigator/Private/Widgets/Columns/SNavigationToolTime.h
+++ b/Engine/Plugins/Experimental/SequenceNavigator/Source/SequenceNavigator/Private/Widgets/Columns/SNavigationToolTime.h
@@ -2,6 +2,7 @@
 
 #pragma once
 
+#include "Items/NavigationToolSequence.h"
 #include "NavigationToolDefines.h"
 #include "ScopedTransaction.h"
 #include "Templates/UniquePtr.h"
@@ -15,19 +16,24 @@ class ISequencer;
 namespace UE::SequenceNavigator
 {
 
+//class FNavigationToolSequence;
 class INavigationToolView;
 class SNavigationToolTreeRow;
 
+/**
+ * Base time column widget.
+ * Subclass this widget to implement a time-based property column widget.
+ */
 class SNavigationToolTime : public SCompoundWidget
 {
 public:
 	SLATE_BEGIN_ARGS(SNavigationToolTime) {}
 	SLATE_END_ARGS()
 
-	virtual ~SNavigationToolTime() override {}
+	virtual ~SNavigationToolTime() override = default;
 
 	void Construct(const FArguments& InArgs
-		, const FNavigationToolItemRef& InItem
+		, const FNavigationToolViewModelPtr& InItem
 		, const TSharedRef<INavigationToolView>& InView
 		, const TSharedRef<SNavigationToolTreeRow>& InRowWidget);
 
@@ -37,6 +43,8 @@ public:
 protected:
 	virtual FName GetStyleName() const;
 
+	//~ Begin SNavigationToolTime
+
 	virtual double GetFrameTimeValue() const = 0;
 
 	virtual void OnFrameTimeValueChanged(const double InNewValue);
@@ -51,12 +59,12 @@ protected:
 
 	virtual FText GetTransactionText() const = 0;
 
-	TSharedPtr<ISequencer> GetSequencer() const;
+	//~ End SNavigationToolTime
 
-	FNavigationToolItemWeakPtr WeakItem;
+	TSharedPtr<ISequencer> GetSequencer() const;
 
+	FNavigationToolViewModelWeakPtr WeakItem;
 	TWeakPtr<INavigationToolView> WeakView;
-
 	TWeakPtr<SNavigationToolTreeRow> WeakRowWidget;
 
 	TUniquePtr<FScopedTransaction> UndoTransaction;
