[SequenceNavigator] Refactor to use Sequencer view models and type macros

diff --git a/Engine/Plugins/Experimental/SequenceNavigator/Source/SequenceNavigator/Private/Widgets/SNavigationToolCustomTextFilterDialog.cpp b/Engine/Plugins/Experimental/SequenceNavigator/Source/SequenceNavigator/Private/Widgets/SNavigationToolCustomTextFilterDialog.cpp
--- a/Engine/Plugins/Experimental/SequenceNavigator/Source/SequenceNavigator/Private/Widgets/SNavigationToolCustomTextFilterDialog.cpp
+++ b/Engine/Plugins/Experimental/SequenceNavigator/Source/SequenceNavigator/Private/Widgets/SNavigationToolCustomTextFilterDialog.cpp
@@ -285,9 +285,9 @@ TSharedRef<SWidget> SNavigationToolCustomTextFilterDialog::ConstructButtonRow()
 	return ButtonBox;
 }
 
-FReply SNavigationToolCustomTextFilterDialog::OnColorBlockMouseButtonDown(const FGeometry& InGeometry, const FPointerEvent& InMouseEvent)
+FReply SNavigationToolCustomTextFilterDialog::OnColorBlockMouseButtonDown(const FGeometry& InGeometry, const FPointerEvent& InPointerEvent)
 {
-	if (InMouseEvent.GetEffectingButton() == EKeys::LeftMouseButton)
+	if (InPointerEvent.GetEffectingButton() == EKeys::LeftMouseButton)
 	{
 		const FOnLinearColorValueChanged ColorValueChangedDelegate = FOnLinearColorValueChanged::CreateLambda([this](const FLinearColor& InNewColor)
 			{
diff --git a/Engine/Plugins/Experimental/SequenceNavigator/Source/SequenceNavigator/Private/Widgets/SNavigationToolCustomTextFilterDialog.h b/Engine/Plugins/Experimental/SequenceNavigator/Source/SequenceNavigator/Private/Widgets/SNavigationToolCustomTextFilterDialog.h
--- a/Engine/Plugins/Experimental/SequenceNavigator/Source/SequenceNavigator/Private/Widgets/SNavigationToolCustomTextFilterDialog.h
+++ b/Engine/Plugins/Experimental/SequenceNavigator/Source/SequenceNavigator/Private/Widgets/SNavigationToolCustomTextFilterDialog.h
@@ -36,7 +36,7 @@ public:
 
 	void Construct(const FArguments& InArgs);
 
-private:
+protected:
 	static void ShowWindow(const TSharedRef<SWindow>& InWindowToShow, const bool bInModal, const TSharedPtr<SWindow>& InParentWindow = nullptr);
 
 	/** Only allow one instance of this dialog at a time. This class is private so we will store the instance here. */
@@ -48,7 +48,7 @@ private:
 	TSharedRef<SWidget> ConstructContentRow(const FText& InLabel, const TSharedRef<SWidget>& InContentWidget);
 	TSharedRef<SWidget> ConstructButtonRow();
 
-	FReply OnColorBlockMouseButtonDown(const FGeometry& InGeometry, const FPointerEvent& InMouseEvent);
+	FReply OnColorBlockMouseButtonDown(const FGeometry& InGeometry, const FPointerEvent& InPointerEvent);
 
 	bool IsEdit() const;
 
diff --git a/Engine/Plugins/Experimental/SequenceNavigator/Source/SequenceNavigator/Private/Widgets/SNavigationToolFilterBar.cpp b/Engine/Plugins/Experimental/SequenceNavigator/Source/SequenceNavigator/Private/Widgets/SNavigationToolFilterBar.cpp
--- a/Engine/Plugins/Experimental/SequenceNavigator/Source/SequenceNavigator/Private/Widgets/SNavigationToolFilterBar.cpp
+++ b/Engine/Plugins/Experimental/SequenceNavigator/Source/SequenceNavigator/Private/Widgets/SNavigationToolFilterBar.cpp
@@ -16,13 +16,13 @@
 #include "Widgets/SNavigationToolCustomTextFilterDialog.h"
 #include "Widgets/SWindow.h"
 
-using namespace UE::Sequencer;
-
 #define LOCTEXT_NAMESPACE "SNavigationToolFilterBar"
 
 namespace UE::SequenceNavigator
 {
 
+using namespace Sequencer;
+
 SNavigationToolFilterBar::~SNavigationToolFilterBar()
 {
 	if (const TSharedPtr<FNavigationToolFilterBar> FilterBar = WeakFilterBar.Pin())
@@ -100,7 +100,7 @@ void SNavigationToolFilterBar::Construct(const FArguments& InArgs, const TWeakPt
 	SetMuted(InArgs._FiltersMuted);
 }
 
-FReply SNavigationToolFilterBar::OnMouseButtonUp(const FGeometry& MyGeometry, const FPointerEvent& MouseEvent)
+FReply SNavigationToolFilterBar::OnMouseButtonUp(const FGeometry& InGeometry, const FPointerEvent& InPointerEvent)
 {
 	const TSharedPtr<FNavigationToolFilterBar> FilterBar = WeakFilterBar.Pin();
 	if (!FilterBar.IsValid())
@@ -108,14 +108,14 @@ FReply SNavigationToolFilterBar::OnMouseButtonUp(const FGeometry& MyGeometry, co
 		return FReply::Unhandled();
 	}
 
-	if (MouseEvent.GetEffectingButton() == EKeys::RightMouseButton)
+	if (InPointerEvent.GetEffectingButton() == EKeys::RightMouseButton)
 	{
-		const FWidgetPath* EventPath = MouseEvent.GetEventPath();
+		const FWidgetPath* EventPath = InPointerEvent.GetEventPath();
 
 		FSlateApplication::Get().PushMenu(AsShared()
 			, EventPath ? *EventPath : FWidgetPath()
 			, ContextMenu->CreateMenu(FilterBar.ToSharedRef())
-			, MouseEvent.GetScreenSpacePosition()
+			, InPointerEvent.GetScreenSpacePosition()
 			, FPopupTransitionEffect(FPopupTransitionEffect::ContextMenu));
 
 		return FReply::Handled().ReleaseMouseCapture();
diff --git a/Engine/Plugins/Experimental/SequenceNavigator/Source/SequenceNavigator/Private/Widgets/SNavigationToolFilterBar.h b/Engine/Plugins/Experimental/SequenceNavigator/Source/SequenceNavigator/Private/Widgets/SNavigationToolFilterBar.h
--- a/Engine/Plugins/Experimental/SequenceNavigator/Source/SequenceNavigator/Private/Widgets/SNavigationToolFilterBar.h
+++ b/Engine/Plugins/Experimental/SequenceNavigator/Source/SequenceNavigator/Private/Widgets/SNavigationToolFilterBar.h
@@ -58,7 +58,7 @@ public:
 	void Construct(const FArguments& InArgs, const TWeakPtr<FNavigationToolFilterBar>& InWeakFilterBar);
 
 	//~ Begin SWidget
-	virtual FReply OnMouseButtonUp(const FGeometry& MyGeometry, const FPointerEvent& MouseEvent) override;
+	virtual FReply OnMouseButtonUp(const FGeometry& InGeometry, const FPointerEvent& InPointerEvent) override;
 	//~ End SWidget
 
 	TSharedPtr<FNavigationToolFilterBar> GetFilterBar() const;
diff --git a/Engine/Plugins/Experimental/SequenceNavigator/Source/SequenceNavigator/Private/Widgets/SNavigationToolTreeRow.cpp b/Engine/Plugins/Experimental/SequenceNavigator/Source/SequenceNavigator/Private/Widgets/SNavigationToolTreeRow.cpp
--- a/Engine/Plugins/Experimental/SequenceNavigator/Source/SequenceNavigator/Private/Widgets/SNavigationToolTreeRow.cpp
+++ b/Engine/Plugins/Experimental/SequenceNavigator/Source/SequenceNavigator/Private/Widgets/SNavigationToolTreeRow.cpp
@@ -2,7 +2,8 @@
 
 #include "Widgets/SNavigationToolTreeRow.h"
 #include "Columns/INavigationToolColumn.h"
-#include "Items/INavigationToolItem.h"
+#include "Items/NavigationToolItem.h"
+#include "NavigationToolDefines.h"
 #include "NavigationToolStyle.h"
 #include "NavigationToolView.h"
 #include "Widgets/SNavigationToolTreeView.h"
@@ -15,43 +16,92 @@ namespace UE::SequenceNavigator
 void SNavigationToolTreeRow::Construct(const FArguments& InArgs
 	, const TSharedRef<FNavigationToolView>& InToolView
 	, const TSharedPtr<SNavigationToolTreeView>& InTreeView
-	, const FNavigationToolItemPtr& InItem)
+	, const FNavigationToolViewModelPtr& InItem)
 {
 	WeakToolView = InToolView;
 	WeakTreeView = InTreeView;
-	Item = InItem;
+	WeakItem = InItem;
 	HighlightText = InArgs._HighlightText;
 
-	SetColorAndOpacity(TAttribute<FLinearColor>::CreateSP(&*InToolView, &FNavigationToolView::GetItemBrushColor, Item));
+	SetColorAndOpacity(TAttribute<FLinearColor>::CreateSP(&*InToolView, &FNavigationToolView::GetItemBrushColor, InItem));
 
 	SMultiColumnTableRow::Construct(FSuperRowType::FArguments()
 			.Style(&FNavigationToolStyle::Get().GetWidgetStyle<FTableRowStyle>(TEXT("TableViewRow")))
+			/*.OnCanAcceptDrop_Lambda([this, InToolView](const FDragDropEvent& InDragDropEvent
+				, const EItemDropZone InDropZone
+				, const FNavigationToolViewModelWeakPtr InWeakTargetItem)
+				{
+					const FNavigationToolViewModelWeakPtr WeakTargetItem = InWeakTargetItem.ImplicitPin();
+					return InToolView->OnCanDrop(InDragDropEvent, InDropZone, WeakTargetItem);
+				})*/
 			.OnCanAcceptDrop(InToolView, &FNavigationToolView::OnCanDrop)
-			.OnDragDetected(InToolView, &FNavigationToolView::OnDragDetected, Item)
-			.OnDragEnter(InToolView, &FNavigationToolView::OnDragEnter, Item)
-			.OnDragLeave(InToolView, &FNavigationToolView::OnDragLeave, Item)
+
+			/*.OnDragDetected_Lambda([this, InToolView](const FGeometry& InGeometry, const FPointerEvent& InPointerEvent)
+				{
+					const FNavigationToolViewModelWeakPtr WeakTargetItem = WeakItem.ImplicitPin();
+					return InToolView->OnDragDetected(InGeometry, InPointerEvent, WeakTargetItem);
+				})*/
+			.OnDragDetected(InToolView, &FNavigationToolView::OnDragDetected, WeakItem)
+
+			/*.OnDragEnter_Lambda([this, InToolView](const FDragDropEvent& InDragDropEvent)
+				{
+					//const FNavigationToolViewModelWeakPtr WeakItem = WeakItem.ImplicitPin();
+					return InToolView->OnDragEnter(InDragDropEvent, WeakItem.ImplicitPin());
+				})*/
+			.OnDragEnter(InToolView, &FNavigationToolView::OnDragEnter, WeakItem)
+
+			/*.OnDragLeave_Lambda([this, InToolView](const FDragDropEvent& InDragDropEvent)
+				{
+					//const FNavigationToolViewModelWeakPtr WeakItem = WeakItem.ImplicitPin();
+					return InToolView->OnDragLeave(InDragDropEvent, WeakItem.ImplicitPin());
+				})*/
+			.OnDragLeave(InToolView, &FNavigationToolView::OnDragLeave, WeakItem)
+	
+			/*.OnAcceptDrop_Lambda([InToolView](const FDragDropEvent& InDragDropEvent
+				, const EItemDropZone InDropZone
+				, const FNavigationToolViewModelWeakPtr InWeakTargetItem)
+				{
+					const FNavigationToolViewModelWeakPtr WeakTargetItem = InWeakTargetItem.ImplicitPin();
+					return InToolView->OnDrop(InDragDropEvent, InDropZone, WeakTargetItem);
+				})*/
 			.OnAcceptDrop(InToolView, &FNavigationToolView::OnDrop)
+
 			.OnDrop(this, &SNavigationToolTreeRow::OnDefaultDrop)
 		, InTreeView.ToSharedRef());
 }
 
 TSharedRef<SWidget> SNavigationToolTreeRow::GenerateWidgetForColumn(const FName& InColumnName)
 {
-	if (Item.IsValid())
+	using namespace Sequencer;
+
+	const FNavigationToolViewModelPtr Item = WeakItem.Pin();
+	if (!Item.IsValid())
 	{
-		const TSharedPtr<FNavigationToolView> ToolView = WeakToolView.Pin();
-		check(ToolView.IsValid());
-		
-		if (const TSharedPtr<INavigationToolColumn> Column = ToolView->GetColumns().FindRef(InColumnName))
-		{
-			return Column->ConstructRowWidget(Item.ToSharedRef(), ToolView.ToSharedRef(), SharedThis(this));
-		}
+		return SNullWidget::NullWidget;
+	}
+
+	const TSharedPtr<FNavigationToolView> ToolView = WeakToolView.Pin();
+	if (!ToolView.IsValid())
+	{
+		return SNullWidget::NullWidget;
 	}
+
+	if (const TSharedPtr<INavigationToolColumn> Column = ToolView->GetColumns().FindRef(InColumnName))
+	{
+		return Column->ConstructRowWidget(Item, ToolView.ToSharedRef(), SharedThis(this));
+	}
+
 	return SNullWidget::NullWidget;
 }
 
-FReply SNavigationToolTreeRow::OnMouseButtonUp(const FGeometry& MyGeometry, const FPointerEvent& MouseEvent)
+FReply SNavigationToolTreeRow::OnMouseButtonUp(const FGeometry& InGeometry, const FPointerEvent& InPointerEvent)
 {
+	const FNavigationToolViewModelPtr Item = WeakItem.Pin();
+	if (!Item.IsValid())
+	{
+		return FReply::Unhandled();
+	}
+
 	const TSharedPtr<FNavigationToolView> ToolView = WeakToolView.Pin();
 	if (!ToolView.IsValid())
 	{
@@ -59,41 +109,44 @@ FReply SNavigationToolTreeRow::OnMouseButtonUp(const FGeometry& MyGeometry, cons
 	}
 
 	//Select Item and the Tree of Children it contains
-	if (MouseEvent.GetEffectingButton() == EKeys::MiddleMouseButton)
+	if (InPointerEvent.GetEffectingButton() == EKeys::MiddleMouseButton)
 	{
 		ENavigationToolItemSelectionFlags Flags = ENavigationToolItemSelectionFlags::IncludeChildren
 			| ENavigationToolItemSelectionFlags::SignalSelectionChange
 			| ENavigationToolItemSelectionFlags::ScrollIntoView;
 
-		if (MouseEvent.IsControlDown())
+		if (InPointerEvent.IsControlDown())
 		{
 			Flags |= ENavigationToolItemSelectionFlags::AppendToCurrentSelection;
 		}
 
-		ToolView->SelectItems({Item}, Flags);
+		ToolView->SelectItems({ Item }, Flags);
 
 		Item->OnSelect();
 
 		return FReply::Handled();
 	}
-	
-	return SMultiColumnTableRow::OnMouseButtonUp(MyGeometry, MouseEvent);
+
+	return SMultiColumnTableRow::OnMouseButtonUp(InGeometry, InPointerEvent);
 }
 
-FReply SNavigationToolTreeRow::OnMouseButtonDoubleClick(const FGeometry& InGeometry, const FPointerEvent& InMouseEvent)
+FReply SNavigationToolTreeRow::OnMouseButtonDoubleClick(const FGeometry& InGeometry, const FPointerEvent& InPointerEvent)
 {
 	// Select Item and the Tree of Children it contains
-	if (InMouseEvent.GetEffectingButton() == EKeys::LeftMouseButton)
+	if (InPointerEvent.GetEffectingButton() == EKeys::LeftMouseButton)
 	{
 		if (const TSharedPtr<FNavigationToolView> ToolView = WeakToolView.Pin())
 		{
-			Item->OnDoubleClick();
+			if (const FNavigationToolViewModelPtr Item = WeakItem.Pin())
+			{
+				Item->OnDoubleClick();
+			}
 
 			return FReply::Handled();
 		}
 	}
 
-	return SMultiColumnTableRow<FNavigationToolItemPtr>::OnMouseButtonDoubleClick(InGeometry, InMouseEvent);
+	return SMultiColumnTableRow::OnMouseButtonDoubleClick(InGeometry, InPointerEvent);
 }
 
 TSharedPtr<FNavigationToolView> SNavigationToolTreeRow::GetToolView() const
diff --git a/Engine/Plugins/Experimental/SequenceNavigator/Source/SequenceNavigator/Private/Widgets/SNavigationToolTreeView.cpp b/Engine/Plugins/Experimental/SequenceNavigator/Source/SequenceNavigator/Private/Widgets/SNavigationToolTreeView.cpp
--- a/Engine/Plugins/Experimental/SequenceNavigator/Source/SequenceNavigator/Private/Widgets/SNavigationToolTreeView.cpp
+++ b/Engine/Plugins/Experimental/SequenceNavigator/Source/SequenceNavigator/Private/Widgets/SNavigationToolTreeView.cpp
@@ -3,11 +3,16 @@
 #include "SNavigationToolTreeView.h"
 #include "Framework/Commands/UICommandList.h"
 #include "Items/INavigationToolItem.h"
+#include "Items/NavigationToolItem.h"
+#include "NavigationToolDefines.h"
 #include "NavigationToolView.h"
+#include "MVVM/ViewModels/ViewModelIterators.h"
 
 namespace UE::SequenceNavigator
 {
 
+using namespace Sequencer;
+
 void SNavigationToolTreeView::Construct(const FArguments& InArgs, const TSharedPtr<FNavigationToolView>& InToolView)
 {
 	WeakToolView = InToolView;
@@ -20,77 +25,81 @@ void SNavigationToolTreeView::Construct(const FArguments& InArgs, const TSharedP
 	STreeView::Construct(InArgs._TreeViewArgs);
 }
 
-int32 SNavigationToolTreeView::GetItemIndex(const FNavigationToolItemPtr& Item) const
+int32 SNavigationToolTreeView::GetItemIndex(const FNavigationToolViewModelWeakPtr& InWeakItem) const
 {
-	if (SListView<FNavigationToolItemPtr>::HasValidItemsSource())
+	if (SListView::HasValidItemsSource())
 	{
-		return SListView<FNavigationToolItemPtr>::GetItems().Find(TListTypeTraits<FNavigationToolItemPtr>::NullableItemTypeConvertToItemType(Item));
+		//const TSharedRef<INavigationToolItem> ItemToFind =
+		//	TListTypeTraits<TSharedRef<INavigationToolItem>>::NullableItemTypeConvertToItemType(InItem);
+		return SListView::GetItems().Find(InWeakItem);
 	}
 	return INDEX_NONE;
 }
 
-void SNavigationToolTreeView::FocusOnItem(const FNavigationToolItemPtr& InItem)
+void SNavigationToolTreeView::FocusOnItem(const FNavigationToolViewModelWeakPtr& InWeakItem)
 {
-	SelectorItem = InItem;
-	RangeSelectionStart = InItem;
+	SelectorItem = InWeakItem;
+	RangeSelectionStart = InWeakItem;
 }
 
-void SNavigationToolTreeView::ScrollItemIntoView(const FNavigationToolItemPtr& InItem)
+void SNavigationToolTreeView::UpdateItemExpansions(const FNavigationToolViewModelWeakPtr& InWeakItem)
 {
-	RequestScrollIntoView(InItem, 0);
-}
-
-void SNavigationToolTreeView::UpdateItemExpansions(const FNavigationToolItemPtr& InItem)
-{
-	if (InItem.IsValid())
+	const FNavigationToolViewModelPtr Item = InWeakItem.Pin();
+	if (!Item.IsValid())
 	{
-		const FSparseItemInfo* const SparseItemInfo = SparseItemInfos.Find(InItem);
-
-		bool bIsExpanded = false;
-		bool bHasExpandedChildren = false;
+		return;
+	}
 
-		if (SparseItemInfo)
-		{
-			bIsExpanded = SparseItemInfo->bIsExpanded;
-			bHasExpandedChildren = SparseItemInfo->bHasExpandedChildren;
-		}
+	const FSparseItemInfo* const SparseItemInfo = SparseItemInfos.Find(InWeakItem);
 
-		//Skip to avoid redundancy
-		if (bIsExpanded || bHasExpandedChildren)
-		{
-			return;
-		}
+	bool bIsExpanded = false;
+	bool bHasExpandedChildren = false;
 
-		TArray<FNavigationToolItemPtr> ItemsToCheck = InItem->GetChildren();
-		while (ItemsToCheck.Num() > 0)
-		{
-			FNavigationToolItemPtr ItemToCheck = ItemsToCheck.Pop();
+	if (SparseItemInfo)
+	{
+		bIsExpanded = SparseItemInfo->bIsExpanded;
+		bHasExpandedChildren = SparseItemInfo->bHasExpandedChildren;
+	}
 
-			if (ItemToCheck.IsValid())
-			{
-				if (IsItemExpanded(ItemToCheck))
-				{
-					bHasExpandedChildren = true;
-					break;
-				}
-				ItemsToCheck.Append(ItemToCheck->GetChildren());
-			}
-		}
+	// Skip to avoid redundancy
+	if (bIsExpanded || bHasExpandedChildren)
+	{
+		return;
+	}
 
-		if (bIsExpanded || bHasExpandedChildren)
+	for (const FNavigationToolViewModelPtr& ChildItem : Item.AsModel()->GetDescendantsOfType<INavigationToolItem>())
+	{
+		if (IsItemExpanded(ChildItem))
 		{
-			SparseItemInfos.Add(InItem, FSparseItemInfo(bIsExpanded, bHasExpandedChildren));
+			bHasExpandedChildren = true;
+			break;
 		}
 	}
+
+	if (bIsExpanded || bHasExpandedChildren)
+	{
+		SparseItemInfos.Add(InWeakItem, FSparseItemInfo(bIsExpanded, bHasExpandedChildren));
+	}
 }
 
-void SNavigationToolTreeView::Private_SetItemSelection(const FNavigationToolItemPtr InItem, bool bShouldBeSelected, bool bWasUserDirected)
+void SNavigationToolTreeView::Private_SetItemSelection(const FNavigationToolViewModelWeakPtr InWeakItem,
+	const bool bInShouldBeSelected, const bool bInWasUserDirected)
 {
+	const FNavigationToolViewModelPtr Item = InWeakItem.Pin();
+	if (!Item.IsValid())
+	{
+		return;
+	}
+
 	const TSharedPtr<FNavigationToolView> ToolView = WeakToolView.Pin();
-	if (ToolView.IsValid() && !ToolView->IsToolLocked() && ToolView->CanSelectItem(InItem))
+	if (!ToolView.IsValid()
+		|| ToolView->IsToolLocked()
+		|| !ToolView->CanSelectItem(Item))
 	{
-		STreeView::Private_SetItemSelection(InItem, bShouldBeSelected, bWasUserDirected);
+		return;
 	}
+
+	STreeView::Private_SetItemSelection(InWeakItem, bInShouldBeSelected, bInWasUserDirected);
 }
 
 void SNavigationToolTreeView::Private_ClearSelection()
@@ -102,23 +111,23 @@ void SNavigationToolTreeView::Private_ClearSelection()
 	}
 }
 
-void SNavigationToolTreeView::Private_SignalSelectionChanged(ESelectInfo::Type SelectInfo)
+void SNavigationToolTreeView::Private_SignalSelectionChanged(const ESelectInfo::Type InSelectInfo)
 {
-	STreeView::Private_SignalSelectionChanged(SelectInfo);
+	STreeView::Private_SignalSelectionChanged(InSelectInfo);
 
 	const TItemSet AddedItems = SelectedItems.Difference(PreviousSelectedItems);
-	for (const FNavigationToolItemPtr& AddedItem : AddedItems)
+	for (const FNavigationToolViewModelWeakPtr& WeakAddedItem : AddedItems)
 	{
-		if (AddedItem.IsValid())
+		if (const FNavigationToolViewModelPtr AddedItem = WeakAddedItem.Pin())
 		{
 			AddedItem->OnItemSelectionChanged(true);
 		}
 	}
 
 	const TItemSet RemovedItems = PreviousSelectedItems.Difference(SelectedItems);
-	for (const FNavigationToolItemPtr& RemovedItem : RemovedItems)
+	for (const FNavigationToolViewModelWeakPtr& WeakRemovedItem : RemovedItems)
 	{
-		if (RemovedItem.IsValid())
+		if (const FNavigationToolViewModelPtr RemovedItem = WeakRemovedItem.Pin())
 		{
 			RemovedItem->OnItemSelectionChanged(false);
 		}
@@ -129,25 +138,27 @@ void SNavigationToolTreeView::Private_SignalSelectionChanged(ESelectInfo::Type S
 void SNavigationToolTreeView::Private_UpdateParentHighlights()
 {
 	this->Private_ClearHighlightedItems();
-	for (TItemSet::TConstIterator SelectedItemIt(this->SelectedItems); SelectedItemIt; ++SelectedItemIt)
+
+	for (const FNavigationToolViewModelWeakPtr& WeakSelectedItem : SelectedItems)
 	{
-		const FNavigationToolItemPtr SelectedItem(*SelectedItemIt);
+		const FNavigationToolViewModelPtr SelectedItem = WeakSelectedItem.Pin();
+		if (!SelectedItem.IsValid())
+		{
+			continue;
+		}
 
-		// Sometimes selection events can come through before the Linearized List is built, so the item may not exist yet.
+		// Sometimes selection events can come through before the linearized list is built, so the item may not exist yet.
 		const int32 ItemIndex = LinearizedItems.Find(SelectedItem);
 		if (ItemIndex == INDEX_NONE)
 		{
-			if (SelectedItem.IsValid())
+			FNavigationToolViewModelPtr ParentItem = SelectedItem->GetParent();
+			while (ParentItem.IsValid())
 			{
-				FNavigationToolItemPtr ParentItem = SelectedItem->GetParent();
-				while (ParentItem.IsValid())
+				if (LinearizedItems.Contains(ParentItem))
 				{
-					if (LinearizedItems.Contains(ParentItem))
-					{
-						this->SetItemHighlighted(ParentItem, true);
-					}
-					ParentItem = ParentItem->GetParent();
+					this->SetItemHighlighted(ParentItem, true);
 				}
+				ParentItem = ParentItem->GetParent();
 			}
 			continue;
 		}
@@ -155,41 +166,41 @@ void SNavigationToolTreeView::Private_UpdateParentHighlights()
 		if (DenseItemInfos.IsValidIndex(ItemIndex))
 		{
 			const FItemInfo& ItemInfo = DenseItemInfos[ItemIndex];
-			
+
 			int32 ParentIndex = ItemInfo.ParentIndex;
-			
+
 			while (ParentIndex != INDEX_NONE)
 			{
-				const FNavigationToolItemPtr& ParentItem = this->LinearizedItems[ParentIndex];
+				const FNavigationToolViewModelWeakPtr ParentItem = this->LinearizedItems[ParentIndex];
 				this->Private_SetItemHighlighted(ParentItem, true);
 
 				const FItemInfo& ParentItemInfo = DenseItemInfos[ParentIndex];
-				
+
 				ParentIndex = ParentItemInfo.ParentIndex;
 			}
 		}
 	}
 }
 
-FCursorReply SNavigationToolTreeView::OnCursorQuery(const FGeometry& MyGeometry, const FPointerEvent& CursorEvent) const
+FCursorReply SNavigationToolTreeView::OnCursorQuery(const FGeometry& InGeometry, const FPointerEvent& InPointerEvent) const
 {
-	if (IsRightClickScrolling() && CursorEvent.IsMouseButtonDown(EKeys::MiddleMouseButton))
+	if (IsRightClickScrolling() && InPointerEvent.IsMouseButtonDown(EKeys::MiddleMouseButton))
 	{
 		// We hide the native cursor as we'll be drawing the software EMouseCursor::GrabHandClosed cursor
 		return FCursorReply::Cursor(EMouseCursor::None);
 	}
-	return STreeView::OnCursorQuery(MyGeometry, CursorEvent);
+	return STreeView::OnCursorQuery(InGeometry, InPointerEvent);
 }
 
-FReply SNavigationToolTreeView::OnMouseMove(const FGeometry& MyGeometry, const FPointerEvent& MouseEvent)
+FReply SNavigationToolTreeView::OnMouseMove(const FGeometry& InGeometry, const FPointerEvent& InPointerEvent)
 {
-	if (MouseEvent.IsMouseButtonDown(EKeys::MiddleMouseButton) && !MouseEvent.IsTouchEvent())
+	if (InPointerEvent.IsMouseButtonDown(EKeys::MiddleMouseButton) && !InPointerEvent.IsTouchEvent())
 	{
 		// We only care about deltas along the scroll axis
-		FTableViewDimensions CursorDeltaDimensions(Orientation, MouseEvent.GetCursorDelta());
+		FTableViewDimensions CursorDeltaDimensions(Orientation, InPointerEvent.GetCursorDelta());
 		CursorDeltaDimensions.LineAxis = 0.f;
 
-		const float ScrollByAmount = CursorDeltaDimensions.ScrollAxis / MyGeometry.Scale;
+		const float ScrollByAmount = CursorDeltaDimensions.ScrollAxis / InGeometry.Scale;
 
 		// If scrolling with the right mouse button, we need to remember how much we scrolled.
 		// If we did not scroll at all, we will bring up the context menu when the mouse is released.
@@ -208,7 +219,7 @@ FReply SNavigationToolTreeView::OnMouseMove(const FGeometry& MyGeometry, const F
 
 			TickScrollDelta -= ScrollByAmount;
 
-			const float AmountScrolled = this->ScrollBy(MyGeometry, -ScrollByAmount, AllowOverscroll);
+			const float AmountScrolled = this->ScrollBy(InGeometry, -ScrollByAmount, AllowOverscroll);
 
 			FReply Reply = FReply::Handled();
 
@@ -217,7 +228,7 @@ FReply SNavigationToolTreeView::OnMouseMove(const FGeometry& MyGeometry, const F
 			if (this->HasMouseCapture() == false)
 			{
 				Reply.CaptureMouse(AsShared()).UseHighPrecisionMouseMovement(AsShared());
-				SoftwareCursorPosition = MyGeometry.AbsoluteToLocal(MouseEvent.GetScreenSpacePosition());
+				SoftwareCursorPosition = InGeometry.AbsoluteToLocal(InPointerEvent.GetScreenSpacePosition());
 				bShowSoftwareCursor    = true;
 			}
 
@@ -231,12 +242,12 @@ FReply SNavigationToolTreeView::OnMouseMove(const FGeometry& MyGeometry, const F
 		}
 	}
 
-	return STreeView::OnMouseMove(MyGeometry, MouseEvent);
+	return STreeView::OnMouseMove(InGeometry, InPointerEvent);
 }
 
-FReply SNavigationToolTreeView::OnMouseButtonUp(const FGeometry& MyGeometry, const FPointerEvent& MouseEvent)
+FReply SNavigationToolTreeView::OnMouseButtonUp(const FGeometry& InGeometry, const FPointerEvent& InPointerEvent)
 {
-	if (MouseEvent.GetEffectingButton() == EKeys::MiddleMouseButton)
+	if (InPointerEvent.GetEffectingButton() == EKeys::MiddleMouseButton)
 	{
 		FReply Reply = FReply::Handled().ReleaseMouseCapture();
 		AmountScrolledWhileRightMouseDown = 0;
@@ -245,8 +256,8 @@ FReply SNavigationToolTreeView::OnMouseButtonUp(const FGeometry& MyGeometry, con
 		// If we have mouse capture, snap the mouse back to the closest location that is within the list's bounds
 		if (HasMouseCapture())
 		{
-			const FSlateRect ListScreenSpaceRect = MyGeometry.GetLayoutBoundingRect();
-			const FVector2D CursorPosition = MyGeometry.LocalToAbsolute(SoftwareCursorPosition);
+			const FSlateRect ListScreenSpaceRect = InGeometry.GetLayoutBoundingRect();
+			const FVector2D CursorPosition = InGeometry.LocalToAbsolute(SoftwareCursorPosition);
 
 			const FIntPoint BestPositionInList(
 					FMath::RoundToInt(FMath::Clamp(CursorPosition.X, ListScreenSpaceRect.Left, ListScreenSpaceRect.Right)),
@@ -258,10 +269,10 @@ FReply SNavigationToolTreeView::OnMouseButtonUp(const FGeometry& MyGeometry, con
 
 		return Reply;
 	}
-	return STreeView::OnMouseButtonUp(MyGeometry, MouseEvent);
+	return STreeView::OnMouseButtonUp(InGeometry, InPointerEvent);
 }
 
-FReply SNavigationToolTreeView::OnKeyDown(const FGeometry& MyGeometry, const FKeyEvent& InKeyEvent)
+FReply SNavigationToolTreeView::OnKeyDown(const FGeometry& InGeometry, const FKeyEvent& InKeyEvent)
 {
 	if (const TSharedPtr<FNavigationToolView> ToolView = WeakToolView.Pin())
 	{
@@ -273,7 +284,7 @@ FReply SNavigationToolTreeView::OnKeyDown(const FGeometry& MyGeometry, const FKe
 			return FReply::Handled();	
 		}
 	}
-	return STreeView::OnKeyDown(MyGeometry, InKeyEvent);
+	return STreeView::OnKeyDown(InGeometry, InKeyEvent);
 }
 
 } // namespace UE::SequenceNavigator
diff --git a/Engine/Plugins/Experimental/SequenceNavigator/Source/SequenceNavigator/Private/Widgets/SNavigationToolTreeView.h b/Engine/Plugins/Experimental/SequenceNavigator/Source/SequenceNavigator/Private/Widgets/SNavigationToolTreeView.h
--- a/Engine/Plugins/Experimental/SequenceNavigator/Source/SequenceNavigator/Private/Widgets/SNavigationToolTreeView.h
+++ b/Engine/Plugins/Experimental/SequenceNavigator/Source/SequenceNavigator/Private/Widgets/SNavigationToolTreeView.h
@@ -2,7 +2,8 @@
 
 #pragma once
 
-#include "NavigationToolDefines.h"
+#include "Items/NavigationToolItem.h"
+#include "MVVM/ViewModelPtr.h"
 #include "Widgets/Views/STreeView.h"
 
 namespace UE::SequenceNavigator
@@ -10,44 +11,44 @@ namespace UE::SequenceNavigator
 
 class FNavigationToolView;
 
-class SNavigationToolTreeView : public STreeView<FNavigationToolItemPtr>
+class SNavigationToolTreeView
+	: public STreeView<FNavigationToolViewModelWeakPtr>
 {
 public:
 	SLATE_BEGIN_ARGS(SNavigationToolTreeView) {}
-		SLATE_ARGUMENT(STreeView<FNavigationToolItemPtr>::FArguments, TreeViewArgs)
+		SLATE_ARGUMENT(STreeView<FNavigationToolViewModelWeakPtr>::FArguments, TreeViewArgs)
 	SLATE_END_ARGS()
 
 	void Construct(const FArguments& InArgs, const TSharedPtr<FNavigationToolView>& InToolView);
 
-	int32 GetItemIndex(const FNavigationToolItemPtr& Item) const;
+	int32 GetItemIndex(const FNavigationToolViewModelWeakPtr& InWeakItem) const;
 
-	void FocusOnItem(const FNavigationToolItemPtr& InItem);
+	void FocusOnItem(const FNavigationToolViewModelWeakPtr& InWeakItem);
 
-	void ScrollItemIntoView(const FNavigationToolItemPtr& InItem);
-
-	void UpdateItemExpansions(const FNavigationToolItemPtr& InItem);
+	void UpdateItemExpansions(const FNavigationToolViewModelWeakPtr& InWeakItem);
 
 	//~ Begin STreeView
 	virtual void Private_UpdateParentHighlights() override;
 	//~ End STreeView
 
 	//~ Begin ITypedTableView
-	virtual void Private_SetItemSelection(const FNavigationToolItemPtr InItem, bool bShouldBeSelected, bool bWasUserDirected = false) override;
-	virtual void Private_SignalSelectionChanged(ESelectInfo::Type SelectInfo) override;
+	virtual void Private_SetItemSelection(const FNavigationToolViewModelWeakPtr InWeakItem
+		, const bool bInShouldBeSelected, const bool bInWasUserDirected = false) override;
+	virtual void Private_SignalSelectionChanged(const ESelectInfo::Type InSelectInfo) override;
 	virtual void Private_ClearSelection() override;
 	//~ End SListView
 
 	//~ Begin SWidget
-	virtual FCursorReply OnCursorQuery(const FGeometry& MyGeometry, const FPointerEvent& CursorEvent) const override;
-	virtual FReply OnMouseMove(const FGeometry& MyGeometry, const FPointerEvent& MouseEvent) override;
-	virtual FReply OnMouseButtonUp(const FGeometry& MyGeometry, const FPointerEvent& MouseEvent) override;
-	virtual FReply OnKeyDown(const FGeometry& MyGeometry, const FKeyEvent& InKeyEvent) override;
+	virtual FCursorReply OnCursorQuery(const FGeometry& InGeometry, const FPointerEvent& InPointerEvent) const override;
+	virtual FReply OnMouseMove(const FGeometry& InGeometry, const FPointerEvent& InPointerEvent) override;
+	virtual FReply OnMouseButtonUp(const FGeometry& InGeometry, const FPointerEvent& InPointerEvent) override;
+	virtual FReply OnKeyDown(const FGeometry& InGeometry, const FKeyEvent& InKeyEvent) override;
 	//~ End SWidget
 
-private:
+protected:
 	TWeakPtr<FNavigationToolView> WeakToolView;
-	
-	SNavigationToolTreeView::TItemSet PreviousSelectedItems;
+
+	TItemSet PreviousSelectedItems;
 };
 
 } // namespace UE::SequenceNavigator
diff --git a/Engine/Plugins/Experimental/SequenceNavigator/Source/SequenceNavigator/Private/Widgets/SNavigationToolView.cpp b/Engine/Plugins/Experimental/SequenceNavigator/Source/SequenceNavigator/Private/Widgets/SNavigationToolView.cpp
--- a/Engine/Plugins/Experimental/SequenceNavigator/Source/SequenceNavigator/Private/Widgets/SNavigationToolView.cpp
+++ b/Engine/Plugins/Experimental/SequenceNavigator/Source/SequenceNavigator/Private/Widgets/SNavigationToolView.cpp
@@ -24,6 +24,8 @@
 namespace UE::SequenceNavigator
 {
 
+using namespace Sequencer;
+
 void ExtractAssetSearchFilterTerms(const FText& SearchText, FString* OutFilterKey, FString* OutFilterValue, int32* OutSuggestionInsertionIndex)
 {
 	const FString SearchString = SearchText.ToString();
@@ -276,7 +278,7 @@ TSharedRef<SWidget> SNavigationToolView::ConstructMainContent()
 			.HAlign(HAlign_Fill)
 			[
 				SAssignNew(TreeView, SNavigationToolTreeView, ToolViewRef)
-				.TreeViewArgs(STreeView<FNavigationToolItemPtr>::FArguments()
+				.TreeViewArgs(STreeView<FNavigationToolViewModelWeakPtr>::FArguments()
 					.HeaderRow(HeaderRowWidget)
 					.TreeItemsSource(&ToolViewRef->GetRootVisibleItems())
 					.OnGetChildren(ToolViewRef, &FNavigationToolView::GetChildrenOfItem)
@@ -286,13 +288,14 @@ TSharedRef<SWidget> SNavigationToolView::ConstructMainContent()
 					.OnContextMenuOpening(ToolViewRef, &FNavigationToolView::CreateItemContextMenu)
 					.OnSetExpansionRecursive(ToolViewRef, &FNavigationToolView::SetItemExpansionRecursive)
 					.HighlightParentNodesForSelection(true)
-					.AllowInvisibleItemSelection(true) //To select items that are still collapsed
+					.AllowInvisibleItemSelection(true) // To select items that are still collapsed
 					.SelectionMode(ESelectionMode::Multi)
-					.OnItemToString_Debug_Lambda([](const FNavigationToolItemPtr& InItem)
+					.OnItemToString_Debug_Lambda([](const FNavigationToolViewModelWeakPtr& InWeakItem)
 						{
-							return InItem->GetItemId().GetStringId();
+							const FNavigationToolViewModelPtr Item = InWeakItem.Pin();
+							return Item.IsValid() ? Item->GetItemId().GetStringId() : FString();
 						})
-				)
+					)
 			]
 		];
 
@@ -426,7 +429,8 @@ void SNavigationToolView::RequestTreeRefresh()
 	TreeView->RequestTreeRefresh();
 }
 
-void SNavigationToolView::SetItemSelection(const TArray<FNavigationToolItemPtr>& InItems, bool bSignalSelectionChange)
+void SNavigationToolView::SetItemSelection(const TArray<FNavigationToolViewModelWeakPtr>& InWeakItems
+	, const bool bInSignalSelectionChange)
 {
 	if (bSelectingItems)
 	{
@@ -437,51 +441,59 @@ void SNavigationToolView::SetItemSelection(const TArray<FNavigationToolItemPtr>&
 
 	TreeView->Private_ClearSelection();
 
-	if (!InItems.IsEmpty())
+	if (!InWeakItems.IsEmpty())
 	{
-		TreeView->SetItemSelection(InItems, true, ESelectInfo::Type::Direct);
+		TreeView->SetItemSelection(InWeakItems, true, ESelectInfo::Type::Direct);
 	}
 
-	if (bSignalSelectionChange)
+	if (bInSignalSelectionChange)
 	{
 		TreeView->Private_SignalSelectionChanged(ESelectInfo::OnMouseClick);
 	}
 }
 
-void SNavigationToolView::OnItemSelectionChanged(const FNavigationToolItemPtr InItem, ESelectInfo::Type InSelectionType)
+void SNavigationToolView::OnItemSelectionChanged(const FNavigationToolViewModelWeakPtr InWeakItem
+	, const ESelectInfo::Type InSelectionType)
 {
-	if (const TSharedPtr<FNavigationToolView> ToolView = WeakToolView.Pin())
+	const TSharedPtr<FNavigationToolView> ToolView = WeakToolView.Pin();
+	if (!ToolView.IsValid())
 	{
-		const TArray<FNavigationToolItemPtr> SelectedItems = TreeView->GetSelectedItems();
-		const bool bUpdateModeTools = InSelectionType != ESelectInfo::Type::Direct;
-		ToolView->NotifyItemSelectionChanged(SelectedItems, InItem, bUpdateModeTools);
+		return;
 	}
+
+	const TArray<FNavigationToolViewModelWeakPtr> WeakSelectedItems = TreeView->GetSelectedItems();
+	const bool bUpdateModeTools = InSelectionType != ESelectInfo::Type::Direct;
+	ToolView->NotifyItemSelectionChanged(WeakSelectedItems, InWeakItem.Pin(), bUpdateModeTools);
 }
 
-void SNavigationToolView::ScrollItemIntoView(const FNavigationToolItemPtr& InItem) const
+void SNavigationToolView::ScrollItemIntoView(const FNavigationToolViewModelWeakPtr& InWeakItem) const
 {
-	TreeView->ScrollItemIntoView(InItem);
+	TreeView->RequestScrollIntoView(InWeakItem);
 }
 
-bool SNavigationToolView::IsItemExpanded(const FNavigationToolItemPtr& InItem) const
+bool SNavigationToolView::IsItemExpanded(const FNavigationToolViewModelWeakPtr& InWeakItem) const
 {
-	return TreeView->IsItemExpanded(InItem);
+	return TreeView->IsItemExpanded(InWeakItem);
 }
 
-void SNavigationToolView::SetItemExpansion(const FNavigationToolItemPtr& Item, bool bInExpand) const
+void SNavigationToolView::SetItemExpansion(const FNavigationToolViewModelWeakPtr& InWeakItem, const bool bInExpand) const
 {
-	TreeView->SetItemExpansion(Item, bInExpand);
+	TreeView->SetItemExpansion(InWeakItem, bInExpand);
 }
 
-void SNavigationToolView::UpdateItemExpansions(const FNavigationToolItemPtr& Item) const
+void SNavigationToolView::UpdateItemExpansions(const FNavigationToolViewModelWeakPtr& InWeakItem) const
 {
-	TreeView->UpdateItemExpansions(Item);
+	TreeView->UpdateItemExpansions(InWeakItem);
 }
 
-TSharedRef<ITableRow> SNavigationToolView::OnItemGenerateRow(const FNavigationToolItemPtr Item, const TSharedRef<STableViewBase>& OwnerTable)
+TSharedRef<ITableRow> SNavigationToolView::OnItemGenerateRow(const FNavigationToolViewModelWeakPtr InWeakItem
+	, const TSharedRef<STableViewBase>& InOwnerTable)
 {
 	const TSharedPtr<FNavigationToolView> ToolView = WeakToolView.Pin();
-	check(Item.IsValid() && ToolView.IsValid());
+	check(ToolView.IsValid());
+
+	const FNavigationToolViewModelPtr Item = InWeakItem.Pin();
+	check(Item.IsValid());
 
 	const TSharedPtr<FNavigationToolFilterBar> FilterBar = GetFilterBar();
 	check(FilterBar.IsValid());
@@ -501,16 +513,16 @@ void SNavigationToolView::SetKeyboardFocus() const
 	}
 }
 
-void SNavigationToolView::OnMouseEnter(const FGeometry& MyGeometry, const FPointerEvent& MouseEvent)
+void SNavigationToolView::OnMouseEnter(const FGeometry& InGeometry, const FPointerEvent& InPointerEvent)
 {
 	if (const TSharedPtr<FNavigationToolView> ToolView = WeakToolView.Pin())
 	{
 		ToolView->UpdateRecentViews();
 	}
-	SCompoundWidget::OnMouseEnter(MyGeometry, MouseEvent);
+	SCompoundWidget::OnMouseEnter(InGeometry, InPointerEvent);
 }
 
-FReply SNavigationToolView::OnKeyDown(const FGeometry& MyGeometry, const FKeyEvent& InKeyEvent)
+FReply SNavigationToolView::OnKeyDown(const FGeometry& InGeometry, const FKeyEvent& InKeyEvent)
 {
 	bIsEnterLastKeyPressed = InKeyEvent.GetKey() == EKeys::Enter;
 	if (const TSharedPtr<FNavigationToolView> ToolView = WeakToolView.Pin())
@@ -523,10 +535,10 @@ FReply SNavigationToolView::OnKeyDown(const FGeometry& MyGeometry, const FKeyEve
 			return FReply::Handled();
 		}
 	}
-	return SCompoundWidget::OnKeyDown(MyGeometry, InKeyEvent);
+	return SCompoundWidget::OnKeyDown(InGeometry, InKeyEvent);
 }
 
-FReply SNavigationToolView::OnDrop(const FGeometry& MyGeometry, const FDragDropEvent& DragDropEvent)
+FReply SNavigationToolView::OnDrop(const FGeometry& InGeometry, const FDragDropEvent& InDragDropEvent)
 {
 	const TSharedPtr<FNavigationToolView> ToolView = WeakToolView.Pin();
 	if (!ToolView.IsValid())
@@ -534,12 +546,12 @@ FReply SNavigationToolView::OnDrop(const FGeometry& MyGeometry, const FDragDropE
 		return FReply::Unhandled();
 	}
 
-	return ToolView->OnDrop(DragDropEvent, EItemDropZone::OntoItem, nullptr);
+	return ToolView->OnDrop(InDragDropEvent, EItemDropZone::OntoItem, nullptr);
 }
 
-void SNavigationToolView::OnDragEnter(const FGeometry& MyGeometry, const FDragDropEvent& DragDropEvent)
+void SNavigationToolView::OnDragEnter(const FGeometry& InGeometry, const FDragDropEvent& InDragDropEvent)
 {
-	if (TSharedPtr<FNavigationToolItemDragDropOp> ItemDragDropOp = DragDropEvent.GetOperationAs<FNavigationToolItemDragDropOp>())
+	if (TSharedPtr<FNavigationToolItemDragDropOp> ItemDragDropOp = InDragDropEvent.GetOperationAs<FNavigationToolItemDragDropOp>())
 	{
 		bool bIsDragFromThisToolView = ItemDragDropOp->GetToolView() == WeakToolView;
 		bool bContainedItemDragDropOp = ItemDragDropOps.Remove(ItemDragDropOp) > 0;
@@ -555,26 +567,26 @@ void SNavigationToolView::OnDragEnter(const FGeometry& MyGeometry, const FDragDr
 
 	if (const TSharedPtr<FNavigationToolView> ToolView = WeakToolView.Pin())
 	{
-		ToolView->OnDragEnter(DragDropEvent, nullptr);
+		ToolView->OnDragEnter(InDragDropEvent, nullptr);
 	}
 }
 
-void SNavigationToolView::OnDragLeave(const FDragDropEvent& DragDropEvent)
+void SNavigationToolView::OnDragLeave(const FDragDropEvent& InDragDropEvent)
 {
-	if (const TSharedPtr<FNavigationToolItemDragDropOp> ItemDragDropOp = DragDropEvent.GetOperationAs<FNavigationToolItemDragDropOp>())
+	if (const TSharedPtr<FNavigationToolItemDragDropOp> ItemDragDropOp = InDragDropEvent.GetOperationAs<FNavigationToolItemDragDropOp>())
 	{
 		ItemDragDropOps.Add(ItemDragDropOp);
 	}
 
 	if (const TSharedPtr<FNavigationToolView> ToolView = WeakToolView.Pin())
 	{
-		ToolView->OnDragLeave(DragDropEvent, nullptr);
+		ToolView->OnDragLeave(InDragDropEvent, nullptr);
 	}
 }
 
-void SNavigationToolView::SetTreeBorderVisibility(bool bVisible)
+void SNavigationToolView::SetTreeBorderVisibility(const bool bInVisible)
 {
-	TreeBorder->SetVisibility(bVisible ? EVisibility::SelfHitTestInvisible : EVisibility::Hidden);
+	TreeBorder->SetVisibility(bInVisible ? EVisibility::SelfHitTestInvisible : EVisibility::Hidden);
 }
 
 void SNavigationToolView::GenerateColumnState(const FName InColumnId, FNavigationToolViewColumnSaveState& OutColumnState)
@@ -630,7 +642,8 @@ void SNavigationToolView::OnFilterBarStateChanged(const bool InIsVisible, const
 	RebuildWidget();
 }
 
-void SNavigationToolView::OnFiltersChanged(const ENavigationToolFilterChange InChangeType, const TSharedRef<FNavigationToolFilter>& InFilter)
+void SNavigationToolView::OnFiltersChanged(const ENavigationToolFilterChange InChangeType
+	, const TSharedRef<FNavigationToolFilter>& InFilter)
 {
 	RebuildWidget();
 }
diff --git a/Engine/Plugins/Experimental/SequenceNavigator/Source/SequenceNavigator/Private/Widgets/SNavigationToolView.h b/Engine/Plugins/Experimental/SequenceNavigator/Source/SequenceNavigator/Private/Widgets/SNavigationToolView.h
--- a/Engine/Plugins/Experimental/SequenceNavigator/Source/SequenceNavigator/Private/Widgets/SNavigationToolView.h
+++ b/Engine/Plugins/Experimental/SequenceNavigator/Source/SequenceNavigator/Private/Widgets/SNavigationToolView.h
@@ -1,4 +1,4 @@
-﻿// Copyright Epic Games, Inc. All Rights Reserved.
+// Copyright Epic Games, Inc. All Rights Reserved.
 
 #pragma once
 
@@ -48,28 +48,33 @@ public:
 
 	void RequestTreeRefresh();
 
-	void SetItemSelection(const TArray<FNavigationToolItemPtr>& InItems, bool bSignalSelectionChange);
+	void SetItemSelection(const TArray<FNavigationToolViewModelWeakPtr>& InWeakItems
+		, const bool bInSignalSelectionChange);
 
-	void OnItemSelectionChanged(const FNavigationToolItemPtr InItem, ESelectInfo::Type InSelectionType);
+	void OnItemSelectionChanged(const FNavigationToolViewModelWeakPtr InWeakItem
+		, const ESelectInfo::Type InSelectionType);
 
-	void ScrollItemIntoView(const FNavigationToolItemPtr& InItem) const;
+	void ScrollItemIntoView(const FNavigationToolViewModelWeakPtr& InWeakItem) const;
 
-	bool IsItemExpanded(const FNavigationToolItemPtr& InItem) const;
-	void SetItemExpansion(const FNavigationToolItemPtr& Item, const bool bInExpand) const;
+	bool IsItemExpanded(const FNavigationToolViewModelWeakPtr& InWeakItem) const;
+	void SetItemExpansion(const FNavigationToolViewModelWeakPtr& InWeakItem, const bool bInExpand) const;
 
-	void UpdateItemExpansions(const FNavigationToolItemPtr& Item) const;
+	void UpdateItemExpansions(const FNavigationToolViewModelWeakPtr& InWeakItem) const;
 
-	TSharedRef<ITableRow> OnItemGenerateRow(const FNavigationToolItemPtr Item, const TSharedRef<STableViewBase>& OwnerTable);
+	TSharedRef<ITableRow> OnItemGenerateRow(const FNavigationToolViewModelWeakPtr InWeakItem
+		, const TSharedRef<STableViewBase>& InOwnerTable);
 
 	void SetKeyboardFocus() const;
 
-	virtual void OnMouseEnter(const FGeometry& MyGeometry, const FPointerEvent& MouseEvent) override;
-	virtual FReply OnKeyDown(const FGeometry& MyGeometry, const FKeyEvent& InKeyEvent) override;
-	virtual FReply OnDrop(const FGeometry& MyGeometry, const FDragDropEvent& DragDropEvent) override;
-	virtual void OnDragEnter(const FGeometry& MyGeometry, const FDragDropEvent& DragDropEvent) override;
-	virtual void OnDragLeave(const FDragDropEvent& DragDropEvent) override;
+	//~ Begin SWidget
+	virtual void OnMouseEnter(const FGeometry& InGeometry, const FPointerEvent& InPointerEvent) override;
+	virtual FReply OnKeyDown(const FGeometry& InGeometry, const FKeyEvent& InKeyEvent) override;
+	virtual FReply OnDrop(const FGeometry& InGeometry, const FDragDropEvent& InDragDropEvent) override;
+	virtual void OnDragEnter(const FGeometry& InGeometry, const FDragDropEvent& InDragDropEvent) override;
+	virtual void OnDragLeave(const FDragDropEvent& InDragDropEvent) override;
+	//~ End SWidget
 
-	void SetTreeBorderVisibility(bool bVisible);
+	void SetTreeBorderVisibility(const bool bInVisible);
 
 	void GenerateColumnState(const FName InColumnId, FNavigationToolViewColumnSaveState& OutColumnState);
 	void GenerateColumnStateMap(TMap<FName, FNavigationToolViewColumnSaveState>& OutStateMap);
@@ -77,14 +82,15 @@ public:
 	EFilterBarLayout GetFilterBarLayout() const;
 	void SetFilterBarLayout(const EFilterBarLayout InLayout);
 
-private:
+protected:
 	TSharedRef<SWidget> ConstructSplitterContent();
 	TSharedRef<SWidget> ConstructMainContent();
 
 	void RebuildSearchAndFilterRow();
 
 	void OnFilterBarStateChanged(const bool InIsVisible, const EFilterBarLayout InNewLayout);
-	void OnFiltersChanged(const ENavigationToolFilterChange InChangeType, const TSharedRef<FNavigationToolFilter>& InFilter);
+	void OnFiltersChanged(const ENavigationToolFilterChange InChangeType
+		, const TSharedRef<FNavigationToolFilter>& InFilter);
 
 	TSharedPtr<FNavigationToolFilterBar> GetFilterBar() const;
 
diff --git a/Engine/Plugins/Experimental/SequenceNavigator/Source/SequenceNavigator/Public/Columns/INavigationToolColumn.h b/Engine/Plugins/Experimental/SequenceNavigator/Source/SequenceNavigator/Public/Columns/INavigationToolColumn.h
--- a/Engine/Plugins/Experimental/SequenceNavigator/Source/SequenceNavigator/Public/Columns/INavigationToolColumn.h
+++ b/Engine/Plugins/Experimental/SequenceNavigator/Source/SequenceNavigator/Public/Columns/INavigationToolColumn.h
@@ -3,19 +3,24 @@
 #pragma once
 
 #include "NavigationToolDefines.h"
-#include "NavigationToolItemType.h"
 #include "Widgets/Views/SHeaderRow.h"
 
+#define UE_API SEQUENCENAVIGATOR_API
+
 namespace UE::SequenceNavigator
 {
 
 class INavigationToolView;
 class SNavigationToolTreeRow;
 
-class INavigationToolColumn : public INavigationToolItemTypeCastable, public TSharedFromThis<INavigationToolColumn>
+class INavigationToolColumn
 {
 public:
-	virtual FName GetColumnId() const { return GetTypeId().ToName(); }
+	UE_SEQUENCER_DECLARE_VIEW_MODEL_TYPE_ID_API(UE_API, INavigationToolColumn)
+
+	virtual ~INavigationToolColumn() = default;
+
+	virtual FName GetColumnId() const = 0;
 
 	virtual FText GetColumnDisplayNameText() const = 0;
 
@@ -33,7 +38,7 @@ public:
 
 	virtual bool CanHideColumn(const FName InColumnId) const { return true; }
 
-	virtual TSharedRef<SWidget> ConstructRowWidget(const FNavigationToolItemRef& InItem
+	virtual TSharedRef<SWidget> ConstructRowWidget(const FNavigationToolViewModelPtr& InItem
 		, const TSharedRef<INavigationToolView>& InView
 		, const TSharedRef<SNavigationToolTreeRow>& InRow) = 0;
 
@@ -41,3 +46,5 @@ public:
 };
 
 } // namespace UE::SequenceNavigator
+
+#undef UE_API
diff --git a/Engine/Plugins/Experimental/SequenceNavigator/Source/SequenceNavigator/Public/Columns/NavigationToolColorColumn.h b/Engine/Plugins/Experimental/SequenceNavigator/Source/SequenceNavigator/Public/Columns/NavigationToolColorColumn.h
--- a/Engine/Plugins/Experimental/SequenceNavigator/Source/SequenceNavigator/Public/Columns/NavigationToolColorColumn.h
+++ b/Engine/Plugins/Experimental/SequenceNavigator/Source/SequenceNavigator/Public/Columns/NavigationToolColorColumn.h
@@ -3,29 +3,37 @@
 #pragma once
 
 #include "Columns/NavigationToolColumn.h"
+#include "Items/NavigationToolItem.h"
+#include "MVVM/ICastable.h"
+
+#define UE_API SEQUENCENAVIGATOR_API
 
 namespace UE::SequenceNavigator
 {
 
 class INavigationToolView;
 
-class FNavigationToolColorColumn : public FNavigationToolColumn
+class FNavigationToolColorColumn
+	: public FNavigationToolColumn
 {
 public:
-	UE_NAVIGATIONTOOL_INHERITS(FNavigationToolColorColumn, FNavigationToolColumn);
+	UE_SEQUENCER_DECLARE_CASTABLE_API(UE_API, FNavigationToolColorColumn
+		, FNavigationToolColumn)
 
 	static FName StaticColumnId() { return TEXT("Color"); }
 
 protected:
 	//~ Begin INavigationToolColumn
 	virtual FName GetColumnId() const override { return StaticColumnId(); }
-	SEQUENCENAVIGATOR_API virtual FText GetColumnDisplayNameText() const override;
+	UE_API virtual FText GetColumnDisplayNameText() const override;
 	virtual bool ShouldShowColumnByDefault() const override { return true; }
-	SEQUENCENAVIGATOR_API virtual SHeaderRow::FColumn::FArguments ConstructHeaderRowColumn(const TSharedRef<INavigationToolView>& InView, const float InFillSize) override;
-	SEQUENCENAVIGATOR_API virtual TSharedRef<SWidget> ConstructRowWidget(const FNavigationToolItemRef& InItem
+	UE_API virtual SHeaderRow::FColumn::FArguments ConstructHeaderRowColumn(const TSharedRef<INavigationToolView>& InView, const float InFillSize) override;
+	UE_API virtual TSharedRef<SWidget> ConstructRowWidget(const FNavigationToolViewModelPtr& InItem
 		, const TSharedRef<INavigationToolView>& InView
 		, const TSharedRef<SNavigationToolTreeRow>& InRow) override;
 	//~ End INavigationToolColumn
 };
 
 } // namespace UE::SequenceNavigator
+
+#undef UE_API
diff --git a/Engine/Plugins/Experimental/SequenceNavigator/Source/SequenceNavigator/Public/Columns/NavigationToolColumn.h b/Engine/Plugins/Experimental/SequenceNavigator/Source/SequenceNavigator/Public/Columns/NavigationToolColumn.h
--- a/Engine/Plugins/Experimental/SequenceNavigator/Source/SequenceNavigator/Public/Columns/NavigationToolColumn.h
+++ b/Engine/Plugins/Experimental/SequenceNavigator/Source/SequenceNavigator/Public/Columns/NavigationToolColumn.h
@@ -3,14 +3,23 @@
 #pragma once
 
 #include "Columns/INavigationToolColumn.h"
+#include "Items/NavigationToolItem.h"
+
+#define UE_API SEQUENCENAVIGATOR_API
 
 namespace UE::SequenceNavigator
 {
 
 /** Base Navigation Tool Column to extend from */
-class FNavigationToolColumn : public INavigationToolColumn
+class FNavigationToolColumn
+	: public Sequencer::FViewModel
+	, public INavigationToolColumn
 {
-
+public:
+	UE_SEQUENCER_DECLARE_CASTABLE_API(UE_API, FNavigationToolColumn
+		, Sequencer::FViewModel, INavigationToolColumn)
 };
 
 } // namespace UE::SequenceNavigator
+
+#undef UE_API
diff --git a/Engine/Plugins/Experimental/SequenceNavigator/Source/SequenceNavigator/Public/Columns/NavigationToolColumnExtender.h b/Engine/Plugins/Experimental/SequenceNavigator/Source/SequenceNavigator/Public/Columns/NavigationToolColumnExtender.h
--- a/Engine/Plugins/Experimental/SequenceNavigator/Source/SequenceNavigator/Public/Columns/NavigationToolColumnExtender.h
+++ b/Engine/Plugins/Experimental/SequenceNavigator/Source/SequenceNavigator/Public/Columns/NavigationToolColumnExtender.h
@@ -5,6 +5,8 @@
 #include "NavigationToolDefines.h"
 #include "Templates/UnrealTypeTraits.h"
 
+#define UE_API SEQUENCENAVIGATOR_API
+
 namespace UE::SequenceNavigator
 {
 
@@ -47,7 +49,7 @@ public:
 	const TArray<TSharedPtr<FNavigationToolColumn>>& GetColumns() const { return Columns; }
 
 private:
-	SEQUENCENAVIGATOR_API void AddColumn(const TSharedPtr<FNavigationToolColumn>& InColumn
+	UE_API void AddColumn(const TSharedPtr<FNavigationToolColumn>& InColumn
 		, const ENavigationToolExtensionPosition InExtensionPosition
 		, const FName InReferenceColumnId);
 
@@ -57,3 +59,5 @@ private:
 };
 
 } // namespace UE::SequenceNavigator
+
+#undef UE_API
diff --git a/Engine/Plugins/Experimental/SequenceNavigator/Source/SequenceNavigator/Public/Columns/NavigationToolCommentColumn.h b/Engine/Plugins/Experimental/SequenceNavigator/Source/SequenceNavigator/Public/Columns/NavigationToolCommentColumn.h
--- a/Engine/Plugins/Experimental/SequenceNavigator/Source/SequenceNavigator/Public/Columns/NavigationToolCommentColumn.h
+++ b/Engine/Plugins/Experimental/SequenceNavigator/Source/SequenceNavigator/Public/Columns/NavigationToolCommentColumn.h
@@ -3,32 +3,40 @@
 #pragma once
 
 #include "Columns/NavigationToolColumn.h"
+#include "Items/NavigationToolItem.h"
+#include "MVVM/ICastable.h"
+
+#define UE_API SEQUENCENAVIGATOR_API
 
 namespace UE::SequenceNavigator
 {
 
 class INavigationToolView;
 
-class FNavigationToolCommentColumn : public FNavigationToolColumn
+class FNavigationToolCommentColumn
+	: public FNavigationToolColumn
 {
 public:
-	UE_NAVIGATIONTOOL_INHERITS(FNavigationToolCommentColumn, FNavigationToolColumn);
+	UE_SEQUENCER_DECLARE_CASTABLE_API(UE_API, FNavigationToolCommentColumn
+		, FNavigationToolColumn)
 
 	static FName StaticColumnId() { return TEXT("Comment"); }
 
 protected:
 	//~ Begin INavigationToolColumn
 	virtual FName GetColumnId() const override { return StaticColumnId(); }
-	SEQUENCENAVIGATOR_API virtual FText GetColumnDisplayNameText() const override;
-	SEQUENCENAVIGATOR_API virtual const FSlateBrush* GetIconBrush() const override;
+	UE_API virtual FText GetColumnDisplayNameText() const override;
+	UE_API virtual const FSlateBrush* GetIconBrush() const override;
 	virtual bool ShouldShowColumnByDefault() const override { return false; }
 	virtual float GetFillWidth() const override { return 5.f; }
-	SEQUENCENAVIGATOR_API virtual SHeaderRow::FColumn::FArguments ConstructHeaderRowColumn(const TSharedRef<INavigationToolView>& InView
+	UE_API virtual SHeaderRow::FColumn::FArguments ConstructHeaderRowColumn(const TSharedRef<INavigationToolView>& InView
 		, const float InFillSize) override;
-	SEQUENCENAVIGATOR_API virtual TSharedRef<SWidget> ConstructRowWidget(const FNavigationToolItemRef& InItem
+	UE_API virtual TSharedRef<SWidget> ConstructRowWidget(const FNavigationToolViewModelPtr& InItem
 		, const TSharedRef<INavigationToolView>& InView
 		, const TSharedRef<SNavigationToolTreeRow>& InRow) override;
 	//~ End INavigationToolColumn
 };
 
 } // namespace UE::SequenceNavigator
+
+#undef UE_API
diff --git a/Engine/Plugins/Experimental/SequenceNavigator/Source/SequenceNavigator/Public/Columns/NavigationToolDeactiveStateColumn.h b/Engine/Plugins/Experimental/SequenceNavigator/Source/SequenceNavigator/Public/Columns/NavigationToolDeactiveStateColumn.h
--- a/Engine/Plugins/Experimental/SequenceNavigator/Source/SequenceNavigator/Public/Columns/NavigationToolDeactiveStateColumn.h
+++ b/Engine/Plugins/Experimental/SequenceNavigator/Source/SequenceNavigator/Public/Columns/NavigationToolDeactiveStateColumn.h
@@ -3,31 +3,39 @@
 #pragma once
 
 #include "Columns/NavigationToolColumn.h"
+#include "Items/NavigationToolItem.h"
+#include "MVVM/ICastable.h"
+
+#define UE_API SEQUENCENAVIGATOR_API
 
 namespace UE::SequenceNavigator
 {
 
 class INavigationToolView;
 
-class FNavigationToolDeactiveStateColumn : public FNavigationToolColumn
+class FNavigationToolDeactiveStateColumn
+	: public FNavigationToolColumn
 {
 public:
-	UE_NAVIGATIONTOOL_INHERITS(FNavigationToolDeactiveStateColumn, FNavigationToolColumn);
+	UE_SEQUENCER_DECLARE_CASTABLE_API(UE_API, FNavigationToolDeactiveStateColumn
+		, FNavigationToolColumn)
 
 	static FName StaticColumnId() { return TEXT("DeactiveState"); }
 
 protected:
 	//~ Begin INavigationToolColumn
 	virtual FName GetColumnId() const override { return StaticColumnId(); }
-	SEQUENCENAVIGATOR_API virtual FText GetColumnDisplayNameText() const override;
-	SEQUENCENAVIGATOR_API virtual const FSlateBrush* GetIconBrush() const override;
-	SEQUENCENAVIGATOR_API virtual SHeaderRow::FColumn::FArguments ConstructHeaderRowColumn(const TSharedRef<INavigationToolView>& InView
+	UE_API virtual FText GetColumnDisplayNameText() const override;
+	UE_API virtual const FSlateBrush* GetIconBrush() const override;
+	UE_API virtual SHeaderRow::FColumn::FArguments ConstructHeaderRowColumn(const TSharedRef<INavigationToolView>& InView
 		, const float InFillSize) override;
 	virtual bool ShouldShowColumnByDefault() const override { return true; }
-	SEQUENCENAVIGATOR_API virtual TSharedRef<SWidget> ConstructRowWidget(const FNavigationToolItemRef& InItem
+	UE_API virtual TSharedRef<SWidget> ConstructRowWidget(const FNavigationToolViewModelPtr& InItem
 		, const TSharedRef<INavigationToolView>& InView
 		, const TSharedRef<SNavigationToolTreeRow>& InRow) override;
 	//~ End INavigationToolColumn
 };
 
 } // namespace UE::SequenceNavigator
+
+#undef UE_API
diff --git a/Engine/Plugins/Experimental/SequenceNavigator/Source/SequenceNavigator/Public/Columns/NavigationToolHBiasColumn.h b/Engine/Plugins/Experimental/SequenceNavigator/Source/SequenceNavigator/Public/Columns/NavigationToolHBiasColumn.h
--- a/Engine/Plugins/Experimental/SequenceNavigator/Source/SequenceNavigator/Public/Columns/NavigationToolHBiasColumn.h
+++ b/Engine/Plugins/Experimental/SequenceNavigator/Source/SequenceNavigator/Public/Columns/NavigationToolHBiasColumn.h
@@ -3,32 +3,40 @@
 #pragma once
 
 #include "Columns/NavigationToolColumn.h"
+#include "Items/NavigationToolItem.h"
+#include "MVVM/ViewModelTypeID.h"
+
+#define UE_API SEQUENCENAVIGATOR_API
 
 namespace UE::SequenceNavigator
 {
 
 class INavigationToolView;
 
-class FNavigationToolHBiasColumn : public FNavigationToolColumn
+class FNavigationToolHBiasColumn
+	: public FNavigationToolColumn
 {
 public:
-	UE_NAVIGATIONTOOL_INHERITS(FNavigationToolHBiasColumn, FNavigationToolColumn);
+	UE_SEQUENCER_DECLARE_CASTABLE_API(UE_API, FNavigationToolHBiasColumn
+		, FNavigationToolColumn)
 
 	static FName StaticColumnId() { return TEXT("HBias"); }
 
 protected:
 	//~ Begin INavigationToolColumn
 	virtual FName GetColumnId() const override { return StaticColumnId(); }
-	SEQUENCENAVIGATOR_API virtual FText GetColumnDisplayNameText() const override;
-	SEQUENCENAVIGATOR_API virtual const FSlateBrush* GetIconBrush() const override;
+	UE_API virtual FText GetColumnDisplayNameText() const override;
+	UE_API virtual const FSlateBrush* GetIconBrush() const override;
 	virtual bool ShouldShowColumnByDefault() const override { return false; }
 	virtual float GetFillWidth() const override { return 5.f; }
-	SEQUENCENAVIGATOR_API virtual SHeaderRow::FColumn::FArguments ConstructHeaderRowColumn(const TSharedRef<INavigationToolView>& InView
+	UE_API virtual SHeaderRow::FColumn::FArguments ConstructHeaderRowColumn(const TSharedRef<INavigationToolView>& InView
 		, const float InFillSize) override;
-	SEQUENCENAVIGATOR_API virtual TSharedRef<SWidget> ConstructRowWidget(const FNavigationToolItemRef& InItem
+	UE_API virtual TSharedRef<SWidget> ConstructRowWidget(const FNavigationToolViewModelPtr& InItem
 		, const TSharedRef<INavigationToolView>& InView
 		, const TSharedRef<SNavigationToolTreeRow>& InRow) override;
 	//~ End INavigationToolColumn
 };
 
 } // namespace UE::SequenceNavigator
+
+#undef UE_API
diff --git a/Engine/Plugins/Experimental/SequenceNavigator/Source/SequenceNavigator/Public/Columns/NavigationToolInTimeColumn.h b/Engine/Plugins/Experimental/SequenceNavigator/Source/SequenceNavigator/Public/Columns/NavigationToolInTimeColumn.h
--- a/Engine/Plugins/Experimental/SequenceNavigator/Source/SequenceNavigator/Public/Columns/NavigationToolInTimeColumn.h
+++ b/Engine/Plugins/Experimental/SequenceNavigator/Source/SequenceNavigator/Public/Columns/NavigationToolInTimeColumn.h
@@ -3,32 +3,40 @@
 #pragma once
 
 #include "Columns/NavigationToolColumn.h"
+#include "Items/NavigationToolItem.h"
+#include "MVVM/ViewModelTypeID.h"
+
+#define UE_API SEQUENCENAVIGATOR_API
 
 namespace UE::SequenceNavigator
 {
 
 class INavigationToolView;
 
-class FNavigationToolInTimeColumn : public FNavigationToolColumn
+class FNavigationToolInTimeColumn
+	: public FNavigationToolColumn
 {
 public:
-	UE_NAVIGATIONTOOL_INHERITS(FNavigationToolInTimeColumn, FNavigationToolColumn);
+	UE_SEQUENCER_DECLARE_CASTABLE_API(UE_API, FNavigationToolInTimeColumn
+		, FNavigationToolColumn)
 
 	static FName StaticColumnId() { return TEXT("InTime"); }
 
 protected:
 	//~ Begin INavigationToolColumn
 	virtual FName GetColumnId() const override { return StaticColumnId(); }
-	SEQUENCENAVIGATOR_API virtual FText GetColumnDisplayNameText() const override;
-	SEQUENCENAVIGATOR_API virtual const FSlateBrush* GetIconBrush() const override;
+	UE_API virtual FText GetColumnDisplayNameText() const override;
+	UE_API virtual const FSlateBrush* GetIconBrush() const override;
 	virtual bool ShouldShowColumnByDefault() const override { return true; }
 	virtual float GetFillWidth() const override { return 5.f; }
-	SEQUENCENAVIGATOR_API virtual SHeaderRow::FColumn::FArguments ConstructHeaderRowColumn(const TSharedRef<INavigationToolView>& InView
+	UE_API virtual SHeaderRow::FColumn::FArguments ConstructHeaderRowColumn(const TSharedRef<INavigationToolView>& InView
 		, const float InFillSize) override;
-	SEQUENCENAVIGATOR_API virtual TSharedRef<SWidget> ConstructRowWidget(const FNavigationToolItemRef& InItem
+	UE_API virtual TSharedRef<SWidget> ConstructRowWidget(const FNavigationToolViewModelPtr& InItem
 		, const TSharedRef<INavigationToolView>& InView
 		, const TSharedRef<SNavigationToolTreeRow>& InRow) override;
 	//~ End INavigationToolColumn
 };
 
 } // namespace UE::SequenceNavigator
+
+#undef UE_API
diff --git a/Engine/Plugins/Experimental/SequenceNavigator/Source/SequenceNavigator/Public/Columns/NavigationToolItemsColumn.h b/Engine/Plugins/Experimental/SequenceNavigator/Source/SequenceNavigator/Public/Columns/NavigationToolItemsColumn.h
--- a/Engine/Plugins/Experimental/SequenceNavigator/Source/SequenceNavigator/Public/Columns/NavigationToolItemsColumn.h
+++ b/Engine/Plugins/Experimental/SequenceNavigator/Source/SequenceNavigator/Public/Columns/NavigationToolItemsColumn.h
@@ -3,32 +3,40 @@
 #pragma once
 
 #include "Columns/NavigationToolColumn.h"
+#include "Items/NavigationToolItem.h"
+#include "MVVM/ICastable.h"
+
+#define UE_API SEQUENCENAVIGATOR_API
 
 namespace UE::SequenceNavigator
 {
 
 class INavigationToolView;
 
-class FNavigationToolItemsColumn : public FNavigationToolColumn
+class FNavigationToolItemsColumn
+	: public FNavigationToolColumn
 {
 public:
-	UE_NAVIGATIONTOOL_INHERITS(FNavigationToolItemsColumn, FNavigationToolColumn);
+	UE_SEQUENCER_DECLARE_CASTABLE_API(UE_API, FNavigationToolItemsColumn
+		, FNavigationToolColumn)
 
 	static FName StaticColumnId() { return TEXT("Items"); }
 
 protected:
 	//~ Begin INavigationToolColumn
 	virtual FName GetColumnId() const override { return StaticColumnId(); }
-	SEQUENCENAVIGATOR_API virtual FText GetColumnDisplayNameText() const override;
-	SEQUENCENAVIGATOR_API virtual const FSlateBrush* GetIconBrush() const override;
+	UE_API virtual FText GetColumnDisplayNameText() const override;
+	UE_API virtual const FSlateBrush* GetIconBrush() const override;
 	virtual bool ShouldShowColumnByDefault() const override { return true; }
 	virtual float GetFillWidth() const override { return 3.f; }
-	SEQUENCENAVIGATOR_API virtual SHeaderRow::FColumn::FArguments ConstructHeaderRowColumn(const TSharedRef<INavigationToolView>& InView
+	UE_API virtual SHeaderRow::FColumn::FArguments ConstructHeaderRowColumn(const TSharedRef<INavigationToolView>& InView
 		, const float InFillSize) override;
-	SEQUENCENAVIGATOR_API virtual TSharedRef<SWidget> ConstructRowWidget(const FNavigationToolItemRef& InItem
+	UE_API virtual TSharedRef<SWidget> ConstructRowWidget(const FNavigationToolViewModelPtr& InItem
 		, const TSharedRef<INavigationToolView>& InView
 		, const TSharedRef<SNavigationToolTreeRow>& InRow) override;
 	//~ End INavigationToolColumn
 };
 
 } // namespace UE::SequenceNavigator
+
+#undef UE_API
diff --git a/Engine/Plugins/Experimental/SequenceNavigator/Source/SequenceNavigator/Public/Columns/NavigationToolLabelColumn.h b/Engine/Plugins/Experimental/SequenceNavigator/Source/SequenceNavigator/Public/Columns/NavigationToolLabelColumn.h
--- a/Engine/Plugins/Experimental/SequenceNavigator/Source/SequenceNavigator/Public/Columns/NavigationToolLabelColumn.h
+++ b/Engine/Plugins/Experimental/SequenceNavigator/Source/SequenceNavigator/Public/Columns/NavigationToolLabelColumn.h
@@ -3,32 +3,40 @@
 #pragma once
 
 #include "Columns/NavigationToolColumn.h"
+#include "Items/NavigationToolItem.h"
+#include "MVVM/ICastable.h"
+
+#define UE_API SEQUENCENAVIGATOR_API
 
 namespace UE::SequenceNavigator
 {
 
 class INavigationToolView;
 
-class FNavigationToolLabelColumn : public FNavigationToolColumn
+class FNavigationToolLabelColumn
+	: public FNavigationToolColumn
 {
 public:
-	UE_NAVIGATIONTOOL_INHERITS(FNavigationToolLabelColumn, FNavigationToolColumn);
+	UE_SEQUENCER_DECLARE_CASTABLE_API(UE_API, FNavigationToolLabelColumn
+		, FNavigationToolColumn)
 
 	static FName StaticColumnId() { return TEXT("Label"); }
 
 protected:
 	//~ Begin INavigationToolColumn
 	virtual FName GetColumnId() const override { return StaticColumnId(); }
-	SEQUENCENAVIGATOR_API virtual FText GetColumnDisplayNameText() const override;
-	SEQUENCENAVIGATOR_API virtual const FSlateBrush* GetIconBrush() const override;
-	SEQUENCENAVIGATOR_API virtual SHeaderRow::FColumn::FArguments ConstructHeaderRowColumn(const TSharedRef<INavigationToolView>& InView
+	UE_API virtual FText GetColumnDisplayNameText() const override;
+	UE_API virtual const FSlateBrush* GetIconBrush() const override;
+	UE_API virtual SHeaderRow::FColumn::FArguments ConstructHeaderRowColumn(const TSharedRef<INavigationToolView>& InView
 		, const float InFillSize) override;
 	virtual float GetFillWidth() const override { return 10.f; }
 	virtual bool CanHideColumn(const FName InColumnId) const override { return false; }
-	SEQUENCENAVIGATOR_API virtual TSharedRef<SWidget> ConstructRowWidget(const FNavigationToolItemRef& InItem
+	UE_API virtual TSharedRef<SWidget> ConstructRowWidget(const FNavigationToolViewModelPtr& InItem
 		, const TSharedRef<INavigationToolView>& InView
 		, const TSharedRef<SNavigationToolTreeRow>& InRow) override;
 	//~ End INavigationToolColumn
 };
 
 } // namespace UE::SequenceNavigator
+
+#undef UE_API
diff --git a/Engine/Plugins/Experimental/SequenceNavigator/Source/SequenceNavigator/Public/Columns/NavigationToolLengthColumn.h b/Engine/Plugins/Experimental/SequenceNavigator/Source/SequenceNavigator/Public/Columns/NavigationToolLengthColumn.h
--- a/Engine/Plugins/Experimental/SequenceNavigator/Source/SequenceNavigator/Public/Columns/NavigationToolLengthColumn.h
+++ b/Engine/Plugins/Experimental/SequenceNavigator/Source/SequenceNavigator/Public/Columns/NavigationToolLengthColumn.h
@@ -3,32 +3,40 @@
 #pragma once
 
 #include "Columns/NavigationToolColumn.h"
+#include "Items/NavigationToolItem.h"
+#include "MVVM/ICastable.h"
+
+#define UE_API SEQUENCENAVIGATOR_API
 
 namespace UE::SequenceNavigator
 {
 
 class INavigationToolView;
 
-class FNavigationToolLengthColumn : public FNavigationToolColumn
+class FNavigationToolLengthColumn
+	: public FNavigationToolColumn
 {
 public:
-	UE_NAVIGATIONTOOL_INHERITS(FNavigationToolLengthColumn, FNavigationToolColumn);
+	UE_SEQUENCER_DECLARE_CASTABLE_API(UE_API, FNavigationToolLengthColumn
+		, FNavigationToolColumn)
 
 	static FName StaticColumnId() { return TEXT("Length"); }
 
 protected:
 	//~ Begin INavigationToolColumn
 	virtual FName GetColumnId() const override { return StaticColumnId(); }
-	SEQUENCENAVIGATOR_API virtual FText GetColumnDisplayNameText() const override;
-	SEQUENCENAVIGATOR_API virtual const FSlateBrush* GetIconBrush() const override;
+	UE_API virtual FText GetColumnDisplayNameText() const override;
+	UE_API virtual const FSlateBrush* GetIconBrush() const override;
 	virtual bool ShouldShowColumnByDefault() const override { return false; }
 	virtual float GetFillWidth() const override { return 5.f; }
-	SEQUENCENAVIGATOR_API virtual SHeaderRow::FColumn::FArguments ConstructHeaderRowColumn(const TSharedRef<INavigationToolView>& InView
+	UE_API virtual SHeaderRow::FColumn::FArguments ConstructHeaderRowColumn(const TSharedRef<INavigationToolView>& InView
 		, const float InFillSize) override;
-	SEQUENCENAVIGATOR_API virtual TSharedRef<SWidget> ConstructRowWidget(const FNavigationToolItemRef& InItem
+	UE_API virtual TSharedRef<SWidget> ConstructRowWidget(const FNavigationToolViewModelPtr& InItem
 		, const TSharedRef<INavigationToolView>& InView
 		, const TSharedRef<SNavigationToolTreeRow>& InRow) override;
 	//~ End INavigationToolColumn
 };
 
 } // namespace UE::SequenceNavigator
+
+#undef UE_API
diff --git a/Engine/Plugins/Experimental/SequenceNavigator/Source/SequenceNavigator/Public/Columns/NavigationToolLockColumn.h b/Engine/Plugins/Experimental/SequenceNavigator/Source/SequenceNavigator/Public/Columns/NavigationToolLockColumn.h
--- a/Engine/Plugins/Experimental/SequenceNavigator/Source/SequenceNavigator/Public/Columns/NavigationToolLockColumn.h
+++ b/Engine/Plugins/Experimental/SequenceNavigator/Source/SequenceNavigator/Public/Columns/NavigationToolLockColumn.h
@@ -3,31 +3,39 @@
 #pragma once
 
 #include "Columns/NavigationToolColumn.h"
+#include "Items/NavigationToolItem.h"
+#include "MVVM/ICastable.h"
+
+#define UE_API SEQUENCENAVIGATOR_API
 
 namespace UE::SequenceNavigator
 {
 
 class INavigationToolView;
 
-class FNavigationToolLockColumn : public FNavigationToolColumn
+class FNavigationToolLockColumn
+	: public FNavigationToolColumn
 {
 public:
-	UE_NAVIGATIONTOOL_INHERITS(FNavigationToolLockColumn, FNavigationToolColumn);
+	UE_SEQUENCER_DECLARE_CASTABLE_API(UE_API, FNavigationToolLockColumn
+		, FNavigationToolColumn)
 
 	static FName StaticColumnId() { return TEXT("Lock"); }
 
 protected:
 	//~ Begin INavigationToolColumn
 	virtual FName GetColumnId() const override { return StaticColumnId(); }
-	SEQUENCENAVIGATOR_API virtual FText GetColumnDisplayNameText() const override;
-	SEQUENCENAVIGATOR_API virtual const FSlateBrush* GetIconBrush() const override;
+	UE_API virtual FText GetColumnDisplayNameText() const override;
+	UE_API virtual const FSlateBrush* GetIconBrush() const override;
 	virtual bool ShouldShowColumnByDefault() const override { return true; }
-	SEQUENCENAVIGATOR_API virtual SHeaderRow::FColumn::FArguments ConstructHeaderRowColumn(const TSharedRef<INavigationToolView>& InView
+	UE_API virtual SHeaderRow::FColumn::FArguments ConstructHeaderRowColumn(const TSharedRef<INavigationToolView>& InView
 		, const float InFillSize) override;
-	SEQUENCENAVIGATOR_API virtual TSharedRef<SWidget> ConstructRowWidget(const FNavigationToolItemRef& InItem
+	UE_API virtual TSharedRef<SWidget> ConstructRowWidget(const FNavigationToolViewModelPtr& InItem
 		, const TSharedRef<INavigationToolView>& InView
 		, const TSharedRef<SNavigationToolTreeRow>& InRow) override;
 	//~ End INavigationToolColumn
 };
 
 } // namespace UE::SequenceNavigator
+
+#undef UE_API
diff --git a/Engine/Plugins/Experimental/SequenceNavigator/Source/SequenceNavigator/Public/Columns/NavigationToolMarkerVisibilityColumn.h b/Engine/Plugins/Experimental/SequenceNavigator/Source/SequenceNavigator/Public/Columns/NavigationToolMarkerVisibilityColumn.h
--- a/Engine/Plugins/Experimental/SequenceNavigator/Source/SequenceNavigator/Public/Columns/NavigationToolMarkerVisibilityColumn.h
+++ b/Engine/Plugins/Experimental/SequenceNavigator/Source/SequenceNavigator/Public/Columns/NavigationToolMarkerVisibilityColumn.h
@@ -3,31 +3,39 @@
 #pragma once
 
 #include "Columns/NavigationToolColumn.h"
+#include "Items/NavigationToolItem.h"
+#include "MVVM/ICastable.h"
+
+#define UE_API SEQUENCENAVIGATOR_API
 
 namespace UE::SequenceNavigator
 {
 
 class INavigationToolView;
 
-class FNavigationToolMarkerVisibilityColumn : public FNavigationToolColumn
+class FNavigationToolMarkerVisibilityColumn
+	: public FNavigationToolColumn
 {
 public:
-	UE_NAVIGATIONTOOL_INHERITS(FNavigationToolMarkerVisibilityColumn, FNavigationToolColumn);
+	UE_SEQUENCER_DECLARE_CASTABLE_API(UE_API, FNavigationToolMarkerVisibilityColumn
+		, FNavigationToolColumn)
 
 	static FName StaticColumnId() { return TEXT("MarkerVisibility"); }
 
 protected:
 	//~ Begin INavigationToolColumn
 	virtual FName GetColumnId() const override { return StaticColumnId(); }
-	SEQUENCENAVIGATOR_API virtual FText GetColumnDisplayNameText() const override;
-	SEQUENCENAVIGATOR_API virtual const FSlateBrush* GetIconBrush() const override;
-	SEQUENCENAVIGATOR_API virtual SHeaderRow::FColumn::FArguments ConstructHeaderRowColumn(const TSharedRef<INavigationToolView>& InView
+	UE_API virtual FText GetColumnDisplayNameText() const override;
+	UE_API virtual const FSlateBrush* GetIconBrush() const override;
+	UE_API virtual SHeaderRow::FColumn::FArguments ConstructHeaderRowColumn(const TSharedRef<INavigationToolView>& InView
 		, const float InFillSize) override;
 	virtual bool ShouldShowColumnByDefault() const override { return true; }
-	SEQUENCENAVIGATOR_API virtual TSharedRef<SWidget> ConstructRowWidget(const FNavigationToolItemRef& InItem
+	UE_API virtual TSharedRef<SWidget> ConstructRowWidget(const FNavigationToolViewModelPtr& InItem
 		, const TSharedRef<INavigationToolView>& InView
 		, const TSharedRef<SNavigationToolTreeRow>& InRow) override;
 	//~ End INavigationToolColumn
 };
 
 } // namespace UE::SequenceNavigator
+
+#undef UE_API
diff --git a/Engine/Plugins/Experimental/SequenceNavigator/Source/SequenceNavigator/Public/Columns/NavigationToolOutTimeColumn.h b/Engine/Plugins/Experimental/SequenceNavigator/Source/SequenceNavigator/Public/Columns/NavigationToolOutTimeColumn.h
--- a/Engine/Plugins/Experimental/SequenceNavigator/Source/SequenceNavigator/Public/Columns/NavigationToolOutTimeColumn.h
+++ b/Engine/Plugins/Experimental/SequenceNavigator/Source/SequenceNavigator/Public/Columns/NavigationToolOutTimeColumn.h
@@ -3,32 +3,40 @@
 #pragma once
 
 #include "Columns/NavigationToolColumn.h"
+#include "Items/NavigationToolItem.h"
+#include "MVVM/ICastable.h"
+
+#define UE_API SEQUENCENAVIGATOR_API
 
 namespace UE::SequenceNavigator
 {
 
 class INavigationToolView;
 
-class FNavigationToolOutTimeColumn : public FNavigationToolColumn
+class FNavigationToolOutTimeColumn
+	: public FNavigationToolColumn
 {
 public:
-	UE_NAVIGATIONTOOL_INHERITS(FNavigationToolOutTimeColumn, FNavigationToolColumn);
+	UE_SEQUENCER_DECLARE_CASTABLE_API(UE_API, FNavigationToolOutTimeColumn
+		, FNavigationToolColumn)
 
 	static FName StaticColumnId() { return TEXT("OutTime"); }
 
 protected:
 	//~ Begin INavigationToolColumn
 	virtual FName GetColumnId() const override { return StaticColumnId(); }
-	SEQUENCENAVIGATOR_API virtual FText GetColumnDisplayNameText() const override;
-	SEQUENCENAVIGATOR_API virtual const FSlateBrush* GetIconBrush() const override;
+	UE_API virtual FText GetColumnDisplayNameText() const override;
+	UE_API virtual const FSlateBrush* GetIconBrush() const override;
 	virtual bool ShouldShowColumnByDefault() const override { return true; }
 	virtual float GetFillWidth() const override { return 5.f; }
-	SEQUENCENAVIGATOR_API virtual SHeaderRow::FColumn::FArguments ConstructHeaderRowColumn(const TSharedRef<INavigationToolView>& InView
+	UE_API virtual SHeaderRow::FColumn::FArguments ConstructHeaderRowColumn(const TSharedRef<INavigationToolView>& InView
 		, const float InFillSize) override;
-	SEQUENCENAVIGATOR_API virtual TSharedRef<SWidget> ConstructRowWidget(const FNavigationToolItemRef& InItem
+	UE_API virtual TSharedRef<SWidget> ConstructRowWidget(const FNavigationToolViewModelPtr& InItem
 		, const TSharedRef<INavigationToolView>& InView
 		, const TSharedRef<SNavigationToolTreeRow>& InRow) override;
 	//~ End INavigationToolColumn
 };
 
 } // namespace UE::SequenceNavigator
+
+#undef UE_API
diff --git a/Engine/Plugins/Experimental/SequenceNavigator/Source/SequenceNavigator/Public/Columns/NavigationToolPlayheadColumn.h b/Engine/Plugins/Experimental/SequenceNavigator/Source/SequenceNavigator/Public/Columns/NavigationToolPlayheadColumn.h
--- a/Engine/Plugins/Experimental/SequenceNavigator/Source/SequenceNavigator/Public/Columns/NavigationToolPlayheadColumn.h
+++ b/Engine/Plugins/Experimental/SequenceNavigator/Source/SequenceNavigator/Public/Columns/NavigationToolPlayheadColumn.h
@@ -3,31 +3,39 @@
 #pragma once
 
 #include "Columns/NavigationToolColumn.h"
+#include "Items/NavigationToolItem.h"
+#include "MVVM/ICastable.h"
+
+#define UE_API SEQUENCENAVIGATOR_API
 
 namespace UE::SequenceNavigator
 {
 
 class INavigationToolView;
 
-class FNavigationToolPlayheadColumn : public FNavigationToolColumn
+class FNavigationToolPlayheadColumn
+	: public FNavigationToolColumn
 {
 public:
-	UE_NAVIGATIONTOOL_INHERITS(FNavigationToolPlayheadColumn, FNavigationToolColumn);
+	UE_SEQUENCER_DECLARE_CASTABLE_API(UE_API, FNavigationToolPlayheadColumn
+		, FNavigationToolColumn)
 
 	static FName StaticColumnId() { return TEXT("Playhead"); }
 
 protected:
 	//~ Begin INavigationToolColumn
 	virtual FName GetColumnId() const override { return StaticColumnId(); }
-	SEQUENCENAVIGATOR_API virtual FText GetColumnDisplayNameText() const override;
-	SEQUENCENAVIGATOR_API virtual const FSlateBrush* GetIconBrush() const override;
-	SEQUENCENAVIGATOR_API virtual SHeaderRow::FColumn::FArguments ConstructHeaderRowColumn(const TSharedRef<INavigationToolView>& InView
+	UE_API virtual FText GetColumnDisplayNameText() const override;
+	UE_API virtual const FSlateBrush* GetIconBrush() const override;
+	UE_API virtual SHeaderRow::FColumn::FArguments ConstructHeaderRowColumn(const TSharedRef<INavigationToolView>& InView
 		, const float InFillSize) override;
 	virtual bool ShouldShowColumnByDefault() const override { return true; }
-	SEQUENCENAVIGATOR_API virtual TSharedRef<SWidget> ConstructRowWidget(const FNavigationToolItemRef& InItem
+	UE_API virtual TSharedRef<SWidget> ConstructRowWidget(const FNavigationToolViewModelPtr& InItem
 		, const TSharedRef<INavigationToolView>& InView
 		, const TSharedRef<SNavigationToolTreeRow>& InRow) override;
 	//~ End INavigationToolColumn
 };
 
 } // namespace UE::SequenceNavigator
+
+#undef UE_API
diff --git a/Engine/Plugins/Experimental/SequenceNavigator/Source/SequenceNavigator/Public/Columns/NavigationToolRevisionControlColumn.h b/Engine/Plugins/Experimental/SequenceNavigator/Source/SequenceNavigator/Public/Columns/NavigationToolRevisionControlColumn.h
--- a/Engine/Plugins/Experimental/SequenceNavigator/Source/SequenceNavigator/Public/Columns/NavigationToolRevisionControlColumn.h
+++ b/Engine/Plugins/Experimental/SequenceNavigator/Source/SequenceNavigator/Public/Columns/NavigationToolRevisionControlColumn.h
@@ -3,31 +3,39 @@
 #pragma once
 
 #include "Columns/NavigationToolColumn.h"
+#include "Items/NavigationToolItem.h"
+#include "MVVM/ICastable.h"
+
+#define UE_API SEQUENCENAVIGATOR_API
 
 namespace UE::SequenceNavigator
 {
 
 class INavigationToolView;
 
-class FNavigationToolRevisionControlColumn : public FNavigationToolColumn
+class FNavigationToolRevisionControlColumn
+	: public FNavigationToolColumn
 {
 public:
-	UE_NAVIGATIONTOOL_INHERITS(FNavigationToolRevisionControlColumn, FNavigationToolColumn);
+	UE_SEQUENCER_DECLARE_CASTABLE_API(UE_API, FNavigationToolRevisionControlColumn
+		, FNavigationToolColumn)
 
 	static FName StaticColumnId() { return TEXT("RevisionControl"); }
 
 protected:
 	//~ Begin INavigationToolColumn
 	virtual FName GetColumnId() const override { return StaticColumnId(); }
-	SEQUENCENAVIGATOR_API virtual FText GetColumnDisplayNameText() const override;
-	SEQUENCENAVIGATOR_API virtual const FSlateBrush* GetIconBrush() const override;
-	SEQUENCENAVIGATOR_API virtual SHeaderRow::FColumn::FArguments ConstructHeaderRowColumn(const TSharedRef<INavigationToolView>& InView
+	UE_API virtual FText GetColumnDisplayNameText() const override;
+	UE_API virtual const FSlateBrush* GetIconBrush() const override;
+	UE_API virtual SHeaderRow::FColumn::FArguments ConstructHeaderRowColumn(const TSharedRef<INavigationToolView>& InView
 		, const float InFillSize) override;
 	virtual bool ShouldShowColumnByDefault() const override { return true; }
-	SEQUENCENAVIGATOR_API virtual TSharedRef<SWidget> ConstructRowWidget(const FNavigationToolItemRef& InItem
+	UE_API virtual TSharedRef<SWidget> ConstructRowWidget(const FNavigationToolViewModelPtr& InItem
 		, const TSharedRef<INavigationToolView>& InView
 		, const TSharedRef<SNavigationToolTreeRow>& InRow) override;
 	//~ End INavigationToolColumn
 };
 
 } // namespace UE::SequenceNavigator
+
+#undef UE_API
diff --git a/Engine/Plugins/Experimental/SequenceNavigator/Source/SequenceNavigator/Public/Columns/NavigationToolStartFrameOffsetColumn.h b/Engine/Plugins/Experimental/SequenceNavigator/Source/SequenceNavigator/Public/Columns/NavigationToolStartFrameOffsetColumn.h
--- a/Engine/Plugins/Experimental/SequenceNavigator/Source/SequenceNavigator/Public/Columns/NavigationToolStartFrameOffsetColumn.h
+++ b/Engine/Plugins/Experimental/SequenceNavigator/Source/SequenceNavigator/Public/Columns/NavigationToolStartFrameOffsetColumn.h
@@ -3,32 +3,40 @@
 #pragma once
 
 #include "Columns/NavigationToolColumn.h"
+#include "Items/NavigationToolItem.h"
+#include "MVVM/ICastable.h"
+
+#define UE_API SEQUENCENAVIGATOR_API
 
 namespace UE::SequenceNavigator
 {
 
 class INavigationToolView;
 
-class FNavigationToolStartFrameOffsetColumn : public FNavigationToolColumn
+class FNavigationToolStartFrameOffsetColumn
+	: public FNavigationToolColumn
 {
 public:
-	UE_NAVIGATIONTOOL_INHERITS(FNavigationToolStartFrameOffsetColumn, FNavigationToolColumn);
+	UE_SEQUENCER_DECLARE_CASTABLE_API(UE_API, FNavigationToolStartFrameOffsetColumn
+		, FNavigationToolColumn)
 
 	static FName StaticColumnId() { return TEXT("StartFrameOffset"); }
 
 protected:
 	//~ Begin INavigationToolColumn
 	virtual FName GetColumnId() const override { return StaticColumnId(); }
-	SEQUENCENAVIGATOR_API virtual FText GetColumnDisplayNameText() const override;
-	SEQUENCENAVIGATOR_API virtual const FSlateBrush* GetIconBrush() const override;
+	UE_API virtual FText GetColumnDisplayNameText() const override;
+	UE_API virtual const FSlateBrush* GetIconBrush() const override;
 	virtual bool ShouldShowColumnByDefault() const override { return false; }
 	virtual float GetFillWidth() const override { return 5.f; }
-	SEQUENCENAVIGATOR_API virtual SHeaderRow::FColumn::FArguments ConstructHeaderRowColumn(const TSharedRef<INavigationToolView>& InView
+	UE_API virtual SHeaderRow::FColumn::FArguments ConstructHeaderRowColumn(const TSharedRef<INavigationToolView>& InView
 		, const float InFillSize) override;
-	SEQUENCENAVIGATOR_API virtual TSharedRef<SWidget> ConstructRowWidget(const FNavigationToolItemRef& InItem
+	UE_API virtual TSharedRef<SWidget> ConstructRowWidget(const FNavigationToolViewModelPtr& InItem
 		, const TSharedRef<INavigationToolView>& InView
 		, const TSharedRef<SNavigationToolTreeRow>& InRow) override;
 	//~ End INavigationToolColumn
 };
 
 } // namespace UE::SequenceNavigator
+
+#undef UE_API
diff --git a/Engine/Plugins/Experimental/SequenceNavigator/Source/SequenceNavigator/Public/Columns/NavigationToolTakeColumn.h b/Engine/Plugins/Experimental/SequenceNavigator/Source/SequenceNavigator/Public/Columns/NavigationToolTakeColumn.h
--- a/Engine/Plugins/Experimental/SequenceNavigator/Source/SequenceNavigator/Public/Columns/NavigationToolTakeColumn.h
+++ b/Engine/Plugins/Experimental/SequenceNavigator/Source/SequenceNavigator/Public/Columns/NavigationToolTakeColumn.h
@@ -3,32 +3,40 @@
 #pragma once
 
 #include "Columns/NavigationToolColumn.h"
+#include "Items/NavigationToolItem.h"
+#include "MVVM/ICastable.h"
+
+#define UE_API SEQUENCENAVIGATOR_API
 
 namespace UE::SequenceNavigator
 {
 
 class INavigationToolView;
 
-class FNavigationToolTakeColumn : public FNavigationToolColumn
+class FNavigationToolTakeColumn
+	: public FNavigationToolColumn
 {
 public:
-	UE_NAVIGATIONTOOL_INHERITS(FNavigationToolTakeColumn, FNavigationToolColumn);
+	UE_SEQUENCER_DECLARE_CASTABLE_API(UE_API, FNavigationToolTakeColumn
+		, FNavigationToolColumn)
 
 	static FName StaticColumnId() { return TEXT("ShotTake"); }
 
 protected:
 	//~ Begin INavigationToolColumn
 	virtual FName GetColumnId() const override { return StaticColumnId(); }
-	SEQUENCENAVIGATOR_API virtual FText GetColumnDisplayNameText() const override;
-	SEQUENCENAVIGATOR_API virtual const FSlateBrush* GetIconBrush() const override;
+	UE_API virtual FText GetColumnDisplayNameText() const override;
+	UE_API virtual const FSlateBrush* GetIconBrush() const override;
 	virtual bool ShouldShowColumnByDefault() const override { return false; }
 	virtual float GetFillWidth() const override { return 5.f; }
-	SEQUENCENAVIGATOR_API virtual SHeaderRow::FColumn::FArguments ConstructHeaderRowColumn(const TSharedRef<INavigationToolView>& InView
+	UE_API virtual SHeaderRow::FColumn::FArguments ConstructHeaderRowColumn(const TSharedRef<INavigationToolView>& InView
 		, const float InFillSize) override;
-	SEQUENCENAVIGATOR_API virtual TSharedRef<SWidget> ConstructRowWidget(const FNavigationToolItemRef& InItem
+	UE_API virtual TSharedRef<SWidget> ConstructRowWidget(const FNavigationToolViewModelPtr& InItem
 		, const TSharedRef<INavigationToolView>& InView
 		, const TSharedRef<SNavigationToolTreeRow>& InRow) override;
 	//~ End INavigationToolColumn
 };
 
 } // namespace UE::SequenceNavigator
+
+#undef UE_API
diff --git a/Engine/Plugins/Experimental/SequenceNavigator/Source/SequenceNavigator/Public/Customization/INavigationToolIconCustomization.h b/Engine/Plugins/Experimental/SequenceNavigator/Source/SequenceNavigator/Public/Customization/INavigationToolIconCustomization.h
--- a/Engine/Plugins/Experimental/SequenceNavigator/Source/SequenceNavigator/Public/Customization/INavigationToolIconCustomization.h
+++ b/Engine/Plugins/Experimental/SequenceNavigator/Source/SequenceNavigator/Public/Customization/INavigationToolIconCustomization.h
@@ -2,6 +2,7 @@
 
 #pragma once
 
+#include "NavigationToolDefines.h"
 #include "Templates/SharedPointer.h"
 #include "UObject/NameTypes.h"
 
@@ -19,9 +20,9 @@ public:
 
 	virtual FName GetToolItemIdentifier() const = 0;
 
-	virtual bool HasOverrideIcon(TSharedPtr<const FNavigationToolItem> InToolItem) const = 0;
+	virtual bool HasOverrideIcon(const FNavigationToolViewModelPtr& InItem) const = 0;
 
-	virtual FSlateIcon GetOverrideIcon(TSharedPtr<const FNavigationToolItem> InToolItem) const = 0;
+	virtual FSlateIcon GetOverrideIcon(const FNavigationToolViewModelPtr& InItem) const = 0;
 };
 
 } // namespace UE::SequenceNavigator
diff --git a/Engine/Plugins/Experimental/SequenceNavigator/Source/SequenceNavigator/Public/DragDropOps/Handlers/NavigationToolItemDropHandler.h b/Engine/Plugins/Experimental/SequenceNavigator/Source/SequenceNavigator/Public/DragDropOps/Handlers/NavigationToolItemDropHandler.h
--- a/Engine/Plugins/Experimental/SequenceNavigator/Source/SequenceNavigator/Public/DragDropOps/Handlers/NavigationToolItemDropHandler.h
+++ b/Engine/Plugins/Experimental/SequenceNavigator/Source/SequenceNavigator/Public/DragDropOps/Handlers/NavigationToolItemDropHandler.h
@@ -3,11 +3,12 @@
 #pragma once
 
 #include "Containers/Array.h"
-#include "Items/INavigationToolItem.h"
+#include "Items/NavigationToolItem.h"
+#include "MVVM/ViewModels/ViewModel.h"
 #include "NavigationToolDefines.h"
-#include "NavigationToolItemType.h"
 #include "Templates/EnableIf.h"
-#include "Templates/SharedPointer.h"
+
+#define UE_API SEQUENCENAVIGATOR_API
 
 enum class EItemDropZone;
 template<typename OptionalType> struct TOptional;
@@ -21,26 +22,29 @@ class FNavigationToolItemDragDropOp;
  * Base Class to Handle Dropping Navigation Tool Items into a Target Navigation Tool Item
  * @see built-in example FNavigationToolActorDropHandler
  */
-class FNavigationToolItemDropHandler : public INavigationToolItemTypeCastable, public TSharedFromThis<FNavigationToolItemDropHandler>
+class FNavigationToolItemDropHandler
+	: public Sequencer::FViewModel
+	//, public TSharedFromThis<FNavigationToolItemDropHandler>
 {
-	friend FNavigationToolItemDragDropOp;
-
-	SEQUENCENAVIGATOR_API void Initialize(const FNavigationToolItemDragDropOp& InDragDropOp);
-
 public:
-	UE_NAVIGATIONTOOL_INHERITS(FNavigationToolItemDropHandler, INavigationToolItemTypeCastable);
+	UE_SEQUENCER_DECLARE_CASTABLE_API(UE_API, FNavigationToolItemDropHandler
+		, Sequencer::FViewModel)
 
-	TConstArrayView<TSharedPtr<INavigationToolItem>> GetItems() const
+	TConstArrayView<FNavigationToolViewModelWeakPtr> GetItems() const
 	{
-		return Items;
+		return WeakItems;
 	}	
 
 protected:
-	virtual bool IsDraggedItemSupported(const FNavigationToolItemPtr& InDraggedItem) const = 0;
+	friend FNavigationToolItemDragDropOp;
+
+	UE_API void Initialize(const FNavigationToolItemDragDropOp& InDragDropOp);
+
+	virtual bool IsDraggedItemSupported(const FNavigationToolViewModelPtr& InDraggedItem) const = 0;
 
-	virtual TOptional<EItemDropZone> CanDrop(EItemDropZone InDropZone, const FNavigationToolItemPtr& InTargetItem) const = 0;
+	virtual TOptional<EItemDropZone> CanDrop(EItemDropZone InDropZone, const FNavigationToolViewModelPtr& InTargetItem) const = 0;
 
-	virtual bool Drop(EItemDropZone InDropZone, const FNavigationToolItemPtr& InTargetItem) = 0;
+	virtual bool Drop(EItemDropZone InDropZone, const FNavigationToolViewModelPtr& InTargetItem) = 0;
 
 	enum class EIterationResult
 	{
@@ -50,27 +54,35 @@ protected:
 	template<typename InItemType, typename = typename TEnableIf<TIsDerivedFrom<InItemType, INavigationToolItem>::Value>::Type>
 	void ForEachItem(const TFunctionRef<EIterationResult(InItemType&)>& InFunc) const
 	{
-		for (const TSharedPtr<INavigationToolItem>& Item : Items)
+		using namespace Sequencer;
+
+		for (const FNavigationToolViewModelWeakPtr& WeakItem : WeakItems)
 		{
+			const FNavigationToolViewModelPtr Item = WeakItem.Pin();
 			if (!Item.IsValid())
 			{
 				continue;
 			}
 
-			if (InItemType* const CastedItem = Item->CastTo<InItemType>())
+			const TViewModelPtr<InItemType> CastedItem = Item.ImplicitCast();
+			if (!CastedItem)
+			{
+				continue;
+			}
+
+			const EIterationResult IterationResult = InFunc(*CastedItem);
+			if (IterationResult == EIterationResult::Break)
 			{
-				EIterationResult IterationResult = InFunc(*CastedItem);
-				if (IterationResult == EIterationResult::Break)
-				{
-					break;
-				}
+				break;
 			}
 		}
 	}
 
-	TArray<FNavigationToolItemPtr> Items;
+	TArray<FNavigationToolViewModelWeakPtr> WeakItems;
 
 	ENavigationToolDragDropActionType ActionType = ENavigationToolDragDropActionType::Move;
 };
 
 } // namespace UE::SequenceNavigator
+
+#undef UE_API
diff --git a/Engine/Plugins/Experimental/SequenceNavigator/Source/SequenceNavigator/Public/DragDropOps/NavigationToolItemDragDropOp.h b/Engine/Plugins/Experimental/SequenceNavigator/Source/SequenceNavigator/Public/DragDropOps/NavigationToolItemDragDropOp.h
--- a/Engine/Plugins/Experimental/SequenceNavigator/Source/SequenceNavigator/Public/DragDropOps/NavigationToolItemDragDropOp.h
+++ b/Engine/Plugins/Experimental/SequenceNavigator/Source/SequenceNavigator/Public/DragDropOps/NavigationToolItemDragDropOp.h
@@ -20,18 +20,18 @@ class FNavigationToolItemDragDropOp : public FDecoratedDragDropOp
 public:
 	DRAG_DROP_OPERATOR_TYPE(FNavigationToolItemDragDropOp, FDecoratedDragDropOp)
 
-	static TSharedRef<FNavigationToolItemDragDropOp> New(const TArray<FNavigationToolItemPtr>& InItems
+	static TSharedRef<FNavigationToolItemDragDropOp> New(const TArray<FNavigationToolViewModelWeakPtr>& InWeakItems
 		, const TSharedPtr<FNavigationToolView>& InToolView
-		, ENavigationToolDragDropActionType InActionType);
+		, const ENavigationToolDragDropActionType InActionType);
 
 	TSharedPtr<INavigationToolView> GetToolView() const
 	{
 		return WeakToolView.Pin();
 	}
 
-	TConstArrayView<FNavigationToolItemPtr> GetItems() const
+	TConstArrayView<FNavigationToolViewModelWeakPtr> GetItems() const
 	{
-		return Items;
+		return WeakItems;
 	}
 
 	ENavigationToolDragDropActionType GetActionType() const
@@ -44,21 +44,21 @@ public:
 		, typename... InArgTypes>
 	void AddDropHandler(InArgTypes&&... InArgs)
 	{
-		TSharedRef<FNavigationToolItemDropHandler> DropHandler = MakeShared<InDropHandlerType>(Forward<InArgTypes>(InArgs)...);
+		const TSharedRef<FNavigationToolItemDropHandler> DropHandler = MakeShared<InDropHandlerType>(Forward<InArgTypes>(InArgs)...);
 		DropHandler->Initialize(*this);
 		DropHandlers.Add(DropHandler);
 	}
 
-	FReply Drop(EItemDropZone InDropZone, const FNavigationToolItemPtr& InTargetItem);
+	FReply Drop(EItemDropZone InDropZone, const FNavigationToolViewModelPtr& InTargetItem);
 
-	TOptional<EItemDropZone> CanDrop(const EItemDropZone InDropZone, const FNavigationToolItemPtr& InTargetItem) const;
+	TOptional<EItemDropZone> CanDrop(const EItemDropZone InDropZone, const FNavigationToolViewModelPtr& InTargetItem) const;
 
 protected:
-	void Init(const TArray<FNavigationToolItemPtr>& InItems
+	void Init(const TArray<FNavigationToolViewModelWeakPtr>& InWeakItems
 		, const TSharedPtr<FNavigationToolView>& InToolView
-		, ENavigationToolDragDropActionType InActionType);
+		, const ENavigationToolDragDropActionType InActionType);
 
-	TArray<FNavigationToolItemPtr> Items;
+	TArray<FNavigationToolViewModelWeakPtr> WeakItems;
 
 	TArray<TSharedRef<FNavigationToolItemDropHandler>> DropHandlers;
 
diff --git a/Engine/Plugins/Experimental/SequenceNavigator/Source/SequenceNavigator/Public/Extensions/IColorExtension.h b/Engine/Plugins/Experimental/SequenceNavigator/Source/SequenceNavigator/Public/Extensions/IColorExtension.h
--- a/Engine/Plugins/Experimental/SequenceNavigator/Source/SequenceNavigator/Public/Extensions/IColorExtension.h
+++ b/Engine/Plugins/Experimental/SequenceNavigator/Source/SequenceNavigator/Public/Extensions/IColorExtension.h
@@ -3,15 +3,19 @@
 #pragma once
 
 #include "Misc/Optional.h"
-#include "NavigationToolItemType.h"
+#include "MVVM/ViewModelTypeID.h"
+
+#define UE_API SEQUENCENAVIGATOR_API
 
 namespace UE::SequenceNavigator
 {
 
-class IColorExtension : public INavigationToolItemTypeCastable
+class IColorExtension
 {
 public:
-	UE_NAVIGATIONTOOL_INHERITS(IColorExtension, INavigationToolItemTypeCastable)
+	UE_SEQUENCER_DECLARE_VIEW_MODEL_TYPE_ID_API(UE_API, IColorExtension)
+
+	virtual ~IColorExtension() = default;
 
 	virtual TOptional<FColor> GetColor() const = 0;
 
@@ -19,3 +23,5 @@ public:
 };
 
 } // namespace UE::SequenceNavigator
+
+#undef UE_API
diff --git a/Engine/Plugins/Experimental/SequenceNavigator/Source/SequenceNavigator/Public/Extensions/IInTimeExtension.h b/Engine/Plugins/Experimental/SequenceNavigator/Source/SequenceNavigator/Public/Extensions/IInTimeExtension.h
--- a/Engine/Plugins/Experimental/SequenceNavigator/Source/SequenceNavigator/Public/Extensions/IInTimeExtension.h
+++ b/Engine/Plugins/Experimental/SequenceNavigator/Source/SequenceNavigator/Public/Extensions/IInTimeExtension.h
@@ -3,15 +3,19 @@
 #pragma once
 
 #include "Misc/FrameNumber.h"
-#include "NavigationToolItemType.h"
+#include "MVVM/ViewModelTypeID.h"
+
+#define UE_API SEQUENCENAVIGATOR_API
 
 namespace UE::SequenceNavigator
 {
 
-class IInTimeExtension : public INavigationToolItemTypeCastable
+class IInTimeExtension
 {
 public:
-	UE_NAVIGATIONTOOL_INHERITS(IInTimeExtension, INavigationToolItemTypeCastable)
+	UE_SEQUENCER_DECLARE_VIEW_MODEL_TYPE_ID_API(UE_API, IInTimeExtension)
+
+	virtual ~IInTimeExtension() = default;
 
 	virtual FFrameNumber GetInTime() const = 0;
 
@@ -19,3 +23,5 @@ public:
 };
 
 } // namespace UE::SequenceNavigator
+
+#undef UE_API
diff --git a/Engine/Plugins/Experimental/SequenceNavigator/Source/SequenceNavigator/Public/Extensions/IMarkerVisibilityExtension.h b/Engine/Plugins/Experimental/SequenceNavigator/Source/SequenceNavigator/Public/Extensions/IMarkerVisibilityExtension.h
--- a/Engine/Plugins/Experimental/SequenceNavigator/Source/SequenceNavigator/Public/Extensions/IMarkerVisibilityExtension.h
+++ b/Engine/Plugins/Experimental/SequenceNavigator/Source/SequenceNavigator/Public/Extensions/IMarkerVisibilityExtension.h
@@ -2,7 +2,9 @@
 
 #pragma once
 
-#include "NavigationToolItemType.h"
+#include "MVVM/ViewModelTypeID.h"
+
+#define UE_API SEQUENCENAVIGATOR_API
 
 namespace UE::SequenceNavigator
 {
@@ -15,10 +17,12 @@ enum class EItemMarkerVisibility
 };
 ENUM_CLASS_FLAGS(EItemMarkerVisibility)
 
-class IMarkerVisibilityExtension : public INavigationToolItemTypeCastable
+class IMarkerVisibilityExtension
 {
 public:
-	UE_NAVIGATIONTOOL_INHERITS(IMarkerVisibilityExtension, INavigationToolItemTypeCastable)
+	UE_SEQUENCER_DECLARE_VIEW_MODEL_TYPE_ID_API(UE_API, IMarkerVisibilityExtension)
+
+	virtual ~IMarkerVisibilityExtension() = default;
 
 	virtual EItemMarkerVisibility GetMarkerVisibility() const = 0;
 
@@ -26,3 +30,5 @@ public:
 };
 
 } // namespace UE::SequenceNavigator
+
+#undef UE_API
diff --git a/Engine/Plugins/Experimental/SequenceNavigator/Source/SequenceNavigator/Public/Extensions/IOutTimeExtension.h b/Engine/Plugins/Experimental/SequenceNavigator/Source/SequenceNavigator/Public/Extensions/IOutTimeExtension.h
--- a/Engine/Plugins/Experimental/SequenceNavigator/Source/SequenceNavigator/Public/Extensions/IOutTimeExtension.h
+++ b/Engine/Plugins/Experimental/SequenceNavigator/Source/SequenceNavigator/Public/Extensions/IOutTimeExtension.h
@@ -3,15 +3,19 @@
 #pragma once
 
 #include "Misc/FrameNumber.h"
-#include "NavigationToolItemType.h"
+#include "MVVM/ViewModelTypeID.h"
+
+#define UE_API SEQUENCENAVIGATOR_API
 
 namespace UE::SequenceNavigator
 {
 
-class IOutTimeExtension : public INavigationToolItemTypeCastable
+class IOutTimeExtension
 {
 public:
-	UE_NAVIGATIONTOOL_INHERITS(IOutTimeExtension, INavigationToolItemTypeCastable)
+	UE_SEQUENCER_DECLARE_VIEW_MODEL_TYPE_ID_API(UE_API, IOutTimeExtension)
+
+	virtual ~IOutTimeExtension() = default;
 
 	virtual FFrameNumber GetOutTime() const = 0;
 
@@ -19,3 +23,5 @@ public:
 };
 
 } // namespace UE::SequenceNavigator
+
+#undef UE_API
diff --git a/Engine/Plugins/Experimental/SequenceNavigator/Source/SequenceNavigator/Public/Extensions/IPlayheadExtension.h b/Engine/Plugins/Experimental/SequenceNavigator/Source/SequenceNavigator/Public/Extensions/IPlayheadExtension.h
--- a/Engine/Plugins/Experimental/SequenceNavigator/Source/SequenceNavigator/Public/Extensions/IPlayheadExtension.h
+++ b/Engine/Plugins/Experimental/SequenceNavigator/Source/SequenceNavigator/Public/Extensions/IPlayheadExtension.h
@@ -2,7 +2,9 @@
 
 #pragma once
 
-#include "NavigationToolItemType.h"
+#include "MVVM/ViewModelTypeID.h"
+
+#define UE_API SEQUENCENAVIGATOR_API
 
 namespace UE::SequenceNavigator
 {
@@ -15,12 +17,16 @@ enum class EItemContainsPlayhead
 };
 ENUM_CLASS_FLAGS(EItemContainsPlayhead)
 
-class IPlayheadExtension : public INavigationToolItemTypeCastable
+class IPlayheadExtension
 {
 public:
-	UE_NAVIGATIONTOOL_INHERITS(IPlayheadExtension, INavigationToolItemTypeCastable)
+	UE_SEQUENCER_DECLARE_VIEW_MODEL_TYPE_ID_API(UE_API, IPlayheadExtension)
+
+	virtual ~IPlayheadExtension() = default;
 
 	virtual EItemContainsPlayhead ContainsPlayhead() const = 0;
 };
 
 } // namespace UE::SequenceNavigator
+
+#undef UE_API
diff --git a/Engine/Plugins/Experimental/SequenceNavigator/Source/SequenceNavigator/Public/Extensions/IRevisionControlExtension.h b/Engine/Plugins/Experimental/SequenceNavigator/Source/SequenceNavigator/Public/Extensions/IRevisionControlExtension.h
--- a/Engine/Plugins/Experimental/SequenceNavigator/Source/SequenceNavigator/Public/Extensions/IRevisionControlExtension.h
+++ b/Engine/Plugins/Experimental/SequenceNavigator/Source/SequenceNavigator/Public/Extensions/IRevisionControlExtension.h
@@ -4,9 +4,11 @@
 
 #include "ISourceControlModule.h"
 #include "Misc/PackageName.h"
-#include "NavigationToolItemType.h"
+#include "MVVM/ViewModelTypeID.h"
 #include "UObject/Package.h"
 
+#define UE_API SEQUENCENAVIGATOR_API
+
 struct FSlateBrush;
 
 namespace UE::SequenceNavigator
@@ -20,15 +22,16 @@ enum class EItemRevisionControlState
 };
 ENUM_CLASS_FLAGS(EItemRevisionControlState)
 
-class IRevisionControlExtension : public INavigationToolItemTypeCastable
+class IRevisionControlExtension
 {
 public:
-	UE_NAVIGATIONTOOL_INHERITS(IRevisionControlExtension, INavigationToolItemTypeCastable)
+	UE_SEQUENCER_DECLARE_VIEW_MODEL_TYPE_ID_API(UE_API, IRevisionControlExtension)
 
 	IRevisionControlExtension(const UObject* const InObject)
 	{
 		QueueRevisionControlStatusUpdate(InObject);
 	}
+	virtual ~IRevisionControlExtension() = default;
 
 	virtual EItemRevisionControlState GetRevisionControlState() const = 0;
 
@@ -56,3 +59,5 @@ public:
 };
 
 } // namespace UE::SequenceNavigator
+
+#undef UE_API
