[SequenceNavigator] Refactor to use Sequencer view models and type macros

diff --git a/Engine/Plugins/Experimental/SequenceNavigator/Source/SequenceNavigator/Public/Filters/Filters/NavigationToolBuiltInFilterParams.h b/Engine/Plugins/Experimental/SequenceNavigator/Source/SequenceNavigator/Public/Filters/Filters/NavigationToolBuiltInFilterParams.h
--- a/Engine/Plugins/Experimental/SequenceNavigator/Source/SequenceNavigator/Public/Filters/Filters/NavigationToolBuiltInFilterParams.h
+++ b/Engine/Plugins/Experimental/SequenceNavigator/Source/SequenceNavigator/Public/Filters/Filters/NavigationToolBuiltInFilterParams.h
@@ -7,27 +7,24 @@
 #include "Styling/SlateBrush.h"
 #include "NavigationToolBuiltInFilterParams.generated.h"
 
-class UObject;
+#define UE_API SEQUENCENAVIGATOR_API
 
-namespace UE::SequenceNavigator
-{
-	class FNavigationToolItemTypeId;
-}
+class UObject;
 
 USTRUCT()
 struct FNavigationToolBuiltInFilterParams
 {
 	GENERATED_BODY()
 
-	SEQUENCENAVIGATOR_API static FNavigationToolBuiltInFilterParams CreateSequenceFilter();
-	SEQUENCENAVIGATOR_API static FNavigationToolBuiltInFilterParams CreateTrackFilter();
-	SEQUENCENAVIGATOR_API static FNavigationToolBuiltInFilterParams CreateBindingFilter();
-	SEQUENCENAVIGATOR_API static FNavigationToolBuiltInFilterParams CreateMarkerFilter();
+	UE_API static FNavigationToolBuiltInFilterParams CreateSequenceFilter();
+	UE_API static FNavigationToolBuiltInFilterParams CreateTrackFilter();
+	UE_API static FNavigationToolBuiltInFilterParams CreateBindingFilter();
+	UE_API static FNavigationToolBuiltInFilterParams CreateMarkerFilter();
 
 	FNavigationToolBuiltInFilterParams() = default;
 
-	SEQUENCENAVIGATOR_API FNavigationToolBuiltInFilterParams(const FName InFilterId
-		, const TSet<UE::SequenceNavigator::FNavigationToolItemTypeId>& InItemClasses
+	UE_API FNavigationToolBuiltInFilterParams(const FName InFilterId
+		, const TSet<UE::Sequencer::FViewModelTypeID>& InItemClasses
 		, const TArray<TSubclassOf<UObject>>& InFilterClasses
 		, const ENavigationToolFilterMode InMode = ENavigationToolFilterMode::MatchesType
 		, const FSlateBrush* const InIconBrush = nullptr
@@ -56,18 +53,18 @@ struct FNavigationToolBuiltInFilterParams
 
 	void SetOverrideIconColor(const FSlateColor& InNewIconColor);
 
-	bool IsValidItemClass(const UE::SequenceNavigator::FNavigationToolItemTypeId InClassTypeId) const;
+	bool IsValidItemClass(const UE::Sequencer::FViewModelTypeID& InClassTypeId) const;
 	bool IsValidObjectClass(const UClass* const InClass) const;
 
 	void SetFilterText(const FText& InText);
 
-	bool PassesFilterText(const UE::SequenceNavigator::FNavigationToolItemPtr& InItem) const;
+	bool PassesFilterText(const UE::SequenceNavigator::FNavigationToolViewModelPtr& InItem) const;
 
 private:
 	UPROPERTY(EditAnywhere, Category = "Filter")
 	FName FilterId = NAME_None;
 
-	TSet<UE::SequenceNavigator::FNavigationToolItemTypeId> ItemClasses;
+	TSet<UE::Sequencer::FViewModelTypeID> ItemClasses;
 
 	UPROPERTY(EditAnywhere, Category = "Filter")
 	TArray<TSubclassOf<UObject>> ObjectClasses;
@@ -100,3 +97,5 @@ private:
 	EClassFlags RequiredClassFlags = CLASS_None;
 	EClassFlags RestrictedClassFlags = CLASS_None;
 };
+
+#undef UE_API
diff --git a/Engine/Plugins/Experimental/SequenceNavigator/Source/SequenceNavigator/Public/Filters/Filters/NavigationToolFilterBase.h b/Engine/Plugins/Experimental/SequenceNavigator/Source/SequenceNavigator/Public/Filters/Filters/NavigationToolFilterBase.h
--- a/Engine/Plugins/Experimental/SequenceNavigator/Source/SequenceNavigator/Public/Filters/Filters/NavigationToolFilterBase.h
+++ b/Engine/Plugins/Experimental/SequenceNavigator/Source/SequenceNavigator/Public/Filters/Filters/NavigationToolFilterBase.h
@@ -5,12 +5,14 @@
 #include "Filters/FilterBase.h"
 #include "Filters/INavigationToolFilterBar.h"
 #include "Filters/SequencerFilterBase.h"
-#include "Items/INavigationToolItem.h"
+#include "Items/NavigationToolItem.h"
 #include "GameFramework/Actor.h"
 #include "Misc/IFilter.h"
 #include "NavigationToolDefines.h"
 #include "NavigationToolFilterBase.generated.h"
 
+#define UE_API SEQUENCENAVIGATOR_API
+
 UENUM()
 enum class ENavigationToolFilterMode : uint8
 {
@@ -27,17 +29,17 @@ ENUM_CLASS_FLAGS(ENavigationToolFilterMode);
 namespace UE::SequenceNavigator
 {
 
-class FNavigationToolFilter : public FSequencerFilterBase<FNavigationToolItemPtr>
+class FNavigationToolFilter : public FSequencerFilterBase<FNavigationToolViewModelPtr>
 {
 public:
-	SEQUENCENAVIGATOR_API FNavigationToolFilter(INavigationToolFilterBar& InOutFilterInterface, TSharedPtr<FFilterCategory>&& InCategory = nullptr);
+	UE_API FNavigationToolFilter(INavigationToolFilterBar& InOutFilterInterface, TSharedPtr<FFilterCategory>&& InCategory = nullptr);
 
 	//~ Begin IFilter
-	virtual bool PassesFilter(const FNavigationToolItemPtr InItem) const override { return true; };
+	virtual bool PassesFilter(const FNavigationToolViewModelPtr InItem) const override { return true; };
 	//~ End IFilter
 
 	//~ Begin FSequencerFilterBase
-	SEQUENCENAVIGATOR_API virtual INavigationToolFilterBar& GetFilterInterface() const override;
+	UE_API virtual INavigationToolFilterBar& GetFilterInterface() const override;
 	//~ End FSequencerFilterBase
 };
 
@@ -54,11 +56,13 @@ public:
 	{}
 
 	//~ Begin IFilter
-	virtual bool PassesFilter(const FNavigationToolItemPtr InItem) const override
+	virtual bool PassesFilter(const FNavigationToolViewModelPtr InItem) const override
 	{
-		return InItem->CastTo<InItemType>() != nullptr;
+		return InItem.AsModel()->CastThis<InItemType>() != nullptr;
 	}
 	//~ End IFilter
 };
 
 } // namespace UE::SequenceNavigator
+
+#undef UE_API
diff --git a/Engine/Plugins/Experimental/SequenceNavigator/Source/SequenceNavigator/Public/Filters/Filters/NavigationToolFilterCollection.h b/Engine/Plugins/Experimental/SequenceNavigator/Source/SequenceNavigator/Public/Filters/Filters/NavigationToolFilterCollection.h
--- a/Engine/Plugins/Experimental/SequenceNavigator/Source/SequenceNavigator/Public/Filters/Filters/NavigationToolFilterCollection.h
+++ b/Engine/Plugins/Experimental/SequenceNavigator/Source/SequenceNavigator/Public/Filters/Filters/NavigationToolFilterCollection.h
@@ -1,4 +1,4 @@
-﻿// Copyright Epic Games, Inc. All Rights Reserved.
+// Copyright Epic Games, Inc. All Rights Reserved.
 
 #pragma once
 
@@ -15,11 +15,10 @@ namespace UE::SequenceNavigator
 class FNavigationToolFilter;
 
 /** Some of this classes functionality could be moved to the TFilterCollection base class so other places could take advantage */
-class FNavigationToolFilterCollection : public TFilterCollection<FNavigationToolItemPtr>
+class FNavigationToolFilterCollection : public TFilterCollection<FNavigationToolViewModelPtr>
 {
 public:
 	FNavigationToolFilterCollection(ISequencerFilterBar& InFilterInterface);
-	virtual ~FNavigationToolFilterCollection() {};
 
 	// @todo Maybe this should get moved in to TFilterCollection
 	bool ContainsFilter(const TSharedRef<FNavigationToolFilter>& InItem) const;
@@ -78,7 +77,7 @@ public:
 
 protected:
 	/**
-	 *	Called when a child Filter restrictions change and broadcasts the FilterChanged delegate
+	 *	Called when a child filter restrictions change and broadcasts the FilterChanged delegate
 	 *	for the collection
 	 */
 	void OnChildFilterChanged();
@@ -90,10 +89,10 @@ public:
 	 * DO NOT USE DIRECTLY
 	 * STL-like iterators to enable range-based for loop support.
 	 */
-	inline TArray<TSharedPtr<IFilter<FNavigationToolItemPtr>>>::RangedForIteratorType      begin()       { return ChildFilters.begin(); }
-	inline TArray<TSharedPtr<IFilter<FNavigationToolItemPtr>>>::RangedForConstIteratorType begin() const { return ChildFilters.begin(); }
-	inline TArray<TSharedPtr<IFilter<FNavigationToolItemPtr>>>::RangedForIteratorType      end()	      { return ChildFilters.end(); }
-	inline TArray<TSharedPtr<IFilter<FNavigationToolItemPtr>>>::RangedForConstIteratorType end()   const { return ChildFilters.end(); }
+	inline TArray<TSharedPtr<IFilter<FNavigationToolViewModelPtr>>>::RangedForIteratorType      begin()       { return ChildFilters.begin(); }
+	inline TArray<TSharedPtr<IFilter<FNavigationToolViewModelPtr>>>::RangedForConstIteratorType begin() const { return ChildFilters.begin(); }
+	inline TArray<TSharedPtr<IFilter<FNavigationToolViewModelPtr>>>::RangedForIteratorType      end()	      { return ChildFilters.end(); }
+	inline TArray<TSharedPtr<IFilter<FNavigationToolViewModelPtr>>>::RangedForConstIteratorType end()   const { return ChildFilters.end(); }
 };
 
 } // namespace UE::SequenceNavigator
diff --git a/Engine/Plugins/Experimental/SequenceNavigator/Source/SequenceNavigator/Public/Filters/Filters/NavigationToolFilter_CustomText.h b/Engine/Plugins/Experimental/SequenceNavigator/Source/SequenceNavigator/Public/Filters/Filters/NavigationToolFilter_CustomText.h
--- a/Engine/Plugins/Experimental/SequenceNavigator/Source/SequenceNavigator/Public/Filters/Filters/NavigationToolFilter_CustomText.h
+++ b/Engine/Plugins/Experimental/SequenceNavigator/Source/SequenceNavigator/Public/Filters/Filters/NavigationToolFilter_CustomText.h
@@ -12,7 +12,7 @@ namespace UE::SequenceNavigator
 class INavigationToolFilterBar;
 
 class FNavigationToolFilter_CustomText :
-	public ICustomTextFilter<FNavigationToolItemPtr>,
+	public ICustomTextFilter<FNavigationToolViewModelPtr>,
 	public FNavigationToolFilter_Text
 {
 public:
@@ -21,7 +21,7 @@ public:
 	//~ Begin ICustomTextFilter
 	virtual void SetFromCustomTextFilterData(const FCustomTextFilterData& InFilterData) override;
 	virtual FCustomTextFilterData CreateCustomTextFilterData() const override;
-	virtual TSharedPtr<FFilterBase<FNavigationToolItemPtr>> GetFilter() override;
+	virtual TSharedPtr<FFilterBase<FNavigationToolViewModelPtr>> GetFilter() override;
 	//~ End ICustomTextFilter
 
 	//~ Begin FNavigationToolFilter
diff --git a/Engine/Plugins/Experimental/SequenceNavigator/Source/SequenceNavigator/Public/Filters/Filters/NavigationToolFilter_Dirty.h b/Engine/Plugins/Experimental/SequenceNavigator/Source/SequenceNavigator/Public/Filters/Filters/NavigationToolFilter_Dirty.h
--- a/Engine/Plugins/Experimental/SequenceNavigator/Source/SequenceNavigator/Public/Filters/Filters/NavigationToolFilter_Dirty.h
+++ b/Engine/Plugins/Experimental/SequenceNavigator/Source/SequenceNavigator/Public/Filters/Filters/NavigationToolFilter_Dirty.h
@@ -27,7 +27,7 @@ public:
 
 	//~ Begin IFilter
 	virtual FString GetName() const override;
-	virtual bool PassesFilter(FNavigationToolFilterType InItem) const override;
+	virtual bool PassesFilter(const FNavigationToolViewModelPtr InItem) const override;
 	//~ End IFilter
 
 protected:
diff --git a/Engine/Plugins/Experimental/SequenceNavigator/Source/SequenceNavigator/Public/Filters/Filters/NavigationToolFilter_Marks.h b/Engine/Plugins/Experimental/SequenceNavigator/Source/SequenceNavigator/Public/Filters/Filters/NavigationToolFilter_Marks.h
--- a/Engine/Plugins/Experimental/SequenceNavigator/Source/SequenceNavigator/Public/Filters/Filters/NavigationToolFilter_Marks.h
+++ b/Engine/Plugins/Experimental/SequenceNavigator/Source/SequenceNavigator/Public/Filters/Filters/NavigationToolFilter_Marks.h
@@ -27,7 +27,7 @@ public:
 
 	//~ Begin IFilter
 	virtual FString GetName() const override;
-	virtual bool PassesFilter(FNavigationToolFilterType InItem) const override;
+	virtual bool PassesFilter(const FNavigationToolViewModelPtr InItem) const override;
 	//~ End IFilter
 };
 
diff --git a/Engine/Plugins/Experimental/SequenceNavigator/Source/SequenceNavigator/Public/Filters/Filters/NavigationToolFilter_Text.h b/Engine/Plugins/Experimental/SequenceNavigator/Source/SequenceNavigator/Public/Filters/Filters/NavigationToolFilter_Text.h
--- a/Engine/Plugins/Experimental/SequenceNavigator/Source/SequenceNavigator/Public/Filters/Filters/NavigationToolFilter_Text.h
+++ b/Engine/Plugins/Experimental/SequenceNavigator/Source/SequenceNavigator/Public/Filters/Filters/NavigationToolFilter_Text.h
@@ -26,7 +26,7 @@ public:
 
 	//~ Begin IFilter
 	virtual FString GetName() const override;
-	virtual bool PassesFilter(const FNavigationToolItemPtr InItem) const override;
+	virtual bool PassesFilter(const FNavigationToolViewModelPtr InItem) const override;
 	//~ End IFilter
 
 	bool IsActive() const;
diff --git a/Engine/Plugins/Experimental/SequenceNavigator/Source/SequenceNavigator/Public/Filters/NavigationToolFilterBar.h b/Engine/Plugins/Experimental/SequenceNavigator/Source/SequenceNavigator/Public/Filters/NavigationToolFilterBar.h
--- a/Engine/Plugins/Experimental/SequenceNavigator/Source/SequenceNavigator/Public/Filters/NavigationToolFilterBar.h
+++ b/Engine/Plugins/Experimental/SequenceNavigator/Source/SequenceNavigator/Public/Filters/NavigationToolFilterBar.h
@@ -54,7 +54,7 @@ public:
 
 	void BindCommands(const TSharedPtr<FUICommandList>& InBaseCommandList);
 
-	TSharedPtr<ICustomTextFilter<FNavigationToolItemPtr>> CreateTextFilter();
+	TSharedPtr<ICustomTextFilter<FNavigationToolViewModelPtr>> CreateTextFilter();
 
 	TSharedRef<SSequencerSearchBox> GetOrCreateSearchBoxWidget();
 
@@ -212,8 +212,8 @@ protected:
 	void CreateDefaultFilters();
 	void CreateCustomTextFiltersFromConfig();
 
-	bool PassesAnyCommonFilter(const FNavigationToolItemPtr& InNode);
-	bool PassesAllCustomTextFilters(const FNavigationToolItemPtr& InNode);
+	bool PassesAnyCommonFilter(const FNavigationToolViewModelPtr& InNode);
+	bool PassesAllCustomTextFilters(const FNavigationToolViewModelPtr& InNode);
 
 	TArray<TSharedRef<FNavigationToolFilter>> GetFilterList(const bool bInIncludeCustomTextFilters = false) const;
 
@@ -263,7 +263,7 @@ private:
 	bool FilterNodesRecursive(INavigationToolView& InToolView
 		, const UNavigationToolSettings& InSettings
 		, const bool bInHasActiveFilter
-		, const FNavigationToolItemPtr& InStartNode);
+		, const FNavigationToolViewModelWeakPtr& InWeakStartNode);
 
 	bool CheckFilterNameValidity(const FString& InNewFilterName, const FString& InOldFilterName, const bool bInIsEdit, FText& OutErrorText) const;
 
diff --git a/Engine/Plugins/Experimental/SequenceNavigator/Source/SequenceNavigator/Public/Filters/NavigationToolFilterData.h b/Engine/Plugins/Experimental/SequenceNavigator/Source/SequenceNavigator/Public/Filters/NavigationToolFilterData.h
--- a/Engine/Plugins/Experimental/SequenceNavigator/Source/SequenceNavigator/Public/Filters/NavigationToolFilterData.h
+++ b/Engine/Plugins/Experimental/SequenceNavigator/Source/SequenceNavigator/Public/Filters/NavigationToolFilterData.h
@@ -2,6 +2,7 @@
 
 #pragma once
 
+#include "Items/NavigationToolItem.h"
 #include "NavigationToolDefines.h"
 #include "UObject/WeakObjectPtr.h"
 
@@ -10,8 +11,6 @@
 namespace UE::SequenceNavigator
 {
 
-using FNavigationToolFilterType = TSharedPtr<INavigationToolItem>;
-
 /** Represents a cache between nodes for a filter operation. */
 struct FNavigationToolFilterData
 {
@@ -32,18 +31,18 @@ struct FNavigationToolFilterData
 
 	UE_API void IncrementTotalNodeCount();
 
-	UE_API void FilterInNode(const FNavigationToolItemPtr& InNode);
-	UE_API void FilterOutNode(const FNavigationToolItemPtr& InNode);
+	UE_API void FilterInNode(const FNavigationToolViewModelWeakPtr& InWeakNode);
+	UE_API void FilterOutNode(const FNavigationToolViewModelWeakPtr& InWeakNode);
 
-	UE_API void FilterInParentChildNodes(const FNavigationToolItemPtr& InNode
+	UE_API void FilterInParentChildNodes(const FNavigationToolViewModelWeakPtr& InWeakNode
 		, const bool bInIncludeSelf, const bool bInIncludeParents, const bool bInIncludeChildren = false);
 
-	UE_API void FilterInNodeWithAncestors(const FNavigationToolItemPtr& InNode);
+	UE_API void FilterInNodeWithAncestors(const FNavigationToolViewModelWeakPtr& InWeakNode);
 
 	UE_API bool ContainsFilterInNodes(const FNavigationToolFilterData& InOtherData) const;
 
-	UE_API bool IsFilteredIn(const FNavigationToolItemPtr& InNode) const;
-	UE_API bool IsFilteredOut(const FNavigationToolItemPtr& InNode) const;
+	UE_API bool IsFilteredIn(const FNavigationToolViewModelWeakPtr& InWeakNode) const;
+	UE_API bool IsFilteredOut(const FNavigationToolViewModelWeakPtr& InWeakNode) const;
 
 protected:
 	FString RawFilterText;
@@ -51,7 +50,7 @@ protected:
 	uint32 TotalNodeCount = 0;
 
 	/** Nodes to be displayed in the UI */
-	TSet<FNavigationToolItemWeakPtr> FilterInNodes;
+	TSet<FNavigationToolViewModelWeakPtr> FilterInNodes;
 };
 
 } // namespace UE::SequenceNavigator
diff --git a/Engine/Plugins/Experimental/SequenceNavigator/Source/SequenceNavigator/Public/Filters/TextExpressions/NavigationToolFilterTextExpressionContext.h b/Engine/Plugins/Experimental/SequenceNavigator/Source/SequenceNavigator/Public/Filters/TextExpressions/NavigationToolFilterTextExpressionContext.h
--- a/Engine/Plugins/Experimental/SequenceNavigator/Source/SequenceNavigator/Public/Filters/TextExpressions/NavigationToolFilterTextExpressionContext.h
+++ b/Engine/Plugins/Experimental/SequenceNavigator/Source/SequenceNavigator/Public/Filters/TextExpressions/NavigationToolFilterTextExpressionContext.h
@@ -7,6 +7,8 @@
 #include "Filters/SequencerTrackFilterBase.h"
 #include "NavigationToolDefines.h"
 
+#define UE_API SEQUENCENAVIGATOR_API
+
 class UMovieScene;
 class UMovieSceneSequence;
 
@@ -16,16 +18,16 @@ namespace UE::SequenceNavigator
 class FNavigationToolFilterTextExpressionContext : public ISequencerTextFilterExpressionContext
 { 
 public:
-	SEQUENCENAVIGATOR_API FNavigationToolFilterTextExpressionContext(INavigationToolFilterBar& InFilterInterface);
+	UE_API FNavigationToolFilterTextExpressionContext(INavigationToolFilterBar& InFilterInterface);
 
-	void SetFilterItem(const FNavigationToolItemPtr& InFilterItem);
+	void SetFilterItem(const FNavigationToolViewModelPtr& InFilterItem);
 
 	//~ Begin ITextFilterExpressionContext
 
-	SEQUENCENAVIGATOR_API virtual bool TestBasicStringExpression(const FTextFilterString& InValue
+	UE_API virtual bool TestBasicStringExpression(const FTextFilterString& InValue
 		, const ETextFilterTextComparisonMode InTextComparisonMode) const override;
 
-	SEQUENCENAVIGATOR_API virtual bool TestComplexExpression(const FName& InKey
+	UE_API virtual bool TestComplexExpression(const FName& InKey
 		, const FTextFilterString& InValue
 		, const ETextFilterComparisonOperation InComparisonOperation
 		, const ETextFilterTextComparisonMode InTextComparisonMode) const override;
@@ -33,12 +35,14 @@ public:
 	//~ End ITextFilterExpressionContext
 
 protected:
-	SEQUENCENAVIGATOR_API UMovieSceneSequence* GetFocusedMovieSceneSequence() const;
-	SEQUENCENAVIGATOR_API UMovieScene* GetFocusedGetMovieScene() const;
+	UE_API UMovieSceneSequence* GetFocusedMovieSceneSequence() const;
+	UE_API UMovieScene* GetFocusedGetMovieScene() const;
 
 	INavigationToolFilterBar& FilterInterface;
 
-	FNavigationToolItemWeakPtr WeakFilterItem;
+	FNavigationToolViewModelWeakPtr WeakFilterItem;
 };
 
 } // namespace UE::SequenceNavigator
+
+#undef UE_API
diff --git a/Engine/Plugins/Experimental/SequenceNavigator/Source/SequenceNavigator/Public/INavigationTool.h b/Engine/Plugins/Experimental/SequenceNavigator/Source/SequenceNavigator/Public/INavigationTool.h
--- a/Engine/Plugins/Experimental/SequenceNavigator/Source/SequenceNavigator/Public/INavigationTool.h
+++ b/Engine/Plugins/Experimental/SequenceNavigator/Source/SequenceNavigator/Public/INavigationTool.h
@@ -8,12 +8,14 @@
 #include "Delegates/Delegate.h"
 #include "ItemProxies/INavigationToolItemProxyFactory.h"
 #include "ItemProxies/NavigationToolItemProxyRegistry.h"
-#include "Items/INavigationToolItem.h"
+#include "Items/NavigationToolItem.h"
 #include "Items/NavigationToolItemId.h"
 #include "Items/NavigationToolItemProxy.h"
 #include "Modules/ModuleManager.h"
+#include "MVVM/ViewModelPtr.h"
 #include "NavigationToolDefines.h"
 #include "NavigationToolExtender.h"
+#include "Providers/NavigationToolProvider.h"
 #include "SequencerCoreFwd.h"
 #include "Templates/SharedPointer.h"
 
@@ -24,8 +26,7 @@ template<typename OptionalType> struct TOptional;
 namespace UE::SequenceNavigator
 {
 
-class FNavigationToolProvider;
-class INavigationToolAction;
+class INavigationToolItemAction;
 class INavigationToolView;
 struct FNavigationToolItemId;
 
@@ -65,17 +66,17 @@ public:
 
 	/** Instantiates a new Item and automatically registers it to the Navigation Tool */
 	template<typename InItemType, typename = typename TEnableIf<TIsDerivedFrom<InItemType, INavigationToolItem>::IsDerived>::Type, typename ...InArgTypes>
-	TSharedRef<InItemType> FindOrAdd(const TSharedPtr<FNavigationToolProvider>& InProvider, InArgTypes&&... InArgs)
+	FNavigationToolViewModelPtr FindOrAdd(const TSharedRef<FNavigationToolProvider>& InProvider, InArgTypes&&... InArgs)
 	{
-		const TSharedRef<InItemType> Item = MakeShared<InItemType>(*this, Forward<InArgTypes>(InArgs)...);
+		const Sequencer::TViewModelPtr<InItemType> Item = MakeShared<InItemType>(*this, Forward<InArgTypes>(InArgs)...);
 		Item->SetProvider(InProvider);
 
 		// If an existing item already exists and has a valid state, use that and forget about the newly created
 		const FNavigationToolItemId ItemId = Item->GetItemId();
-		const FNavigationToolItemPtr ExistingItem = FindItem(ItemId);
-		if (ExistingItem.IsValid() && ExistingItem->IsItemValid() && ExistingItem->IsA<InItemType>())
+		const FNavigationToolViewModelPtr ExistingItem = FindItem(ItemId);
+		if (ExistingItem.IsValid() && ExistingItem->IsItemValid() && ExistingItem.AsModel()->IsA<InItemType>())
 		{
-			return StaticCastSharedPtr<InItemType>(ExistingItem).ToSharedRef();
+			return ExistingItem;
 		}
 
 		if (Item->IsAllowedInTool())
@@ -106,25 +107,25 @@ public:
 	 * @returns the Item Proxy created by the Factory. Can be null if no factory was found or if the factory intentionally returns null
 	 */
 	template<typename InItemProxyType, typename = typename TEnableIf<TIsDerivedFrom<InItemProxyType, FNavigationToolItemProxy>::IsDerived>::Type, typename ...InArgTypes>
-	TSharedPtr<FNavigationToolItemProxy> GetOrCreateItemProxy(const FNavigationToolItemPtr& InParentItem, InArgTypes&&... InArgs)
+	TSharedPtr<FNavigationToolItemProxy> GetOrCreateItemProxy(const FNavigationToolViewModelPtr& InParentItem, InArgTypes&&... InArgs)
 	{
 		if (!InParentItem.IsValid() || !InParentItem->IsAllowedInTool())
 		{
 			return nullptr;
 		}
-		
+
 		INavigationToolItemProxyFactory* const Factory = GetItemProxyFactory<InItemProxyType>();
 		if (!Factory)
 		{
 			return nullptr;
 		}
-		
+
 		TSharedPtr<FNavigationToolItemProxy> OutItemProxy;
-		if (const FNavigationToolItemPtr ExistingItemProxy = FindItem(FNavigationToolItemId(InParentItem, *Factory)))
+		if (const FNavigationToolViewModelPtr ExistingItemProxy = FindItem(FNavigationToolItemId(InParentItem, *Factory)))
 		{
-			check(ExistingItemProxy->IsA<FNavigationToolItemProxy>());
+			check(ExistingItemProxy.AsModel()->IsA<FNavigationToolItemProxy>());
 			ExistingItemProxy->SetParent(InParentItem);
-			OutItemProxy = StaticCastSharedPtr<FNavigationToolItemProxy>(ExistingItemProxy);
+			OutItemProxy = ExistingItemProxy.ImplicitCast();
 		}
 		else
 		{
@@ -144,7 +145,7 @@ public:
 	virtual bool IsObjectAllowedInTool(const UObject* const InObject) const = 0;
 
 	/** Registers the given Item, replacing the old one. */
-	virtual void RegisterItem(const FNavigationToolItemPtr& InItem) = 0;
+	virtual void RegisterItem(const FNavigationToolViewModelPtr& InItem) = 0;
 
 	/** Unregisters the Item having the given ItemId */
 	virtual void UnregisterItem(const FNavigationToolItemId& InItemId) = 0;
@@ -159,21 +160,21 @@ public:
 	 */
 	virtual void Refresh() = 0;
 
-	/** Gets the Tree Root Item of the Navigation Tool */
-	virtual TSharedRef<FNavigationToolItem> GetTreeRoot() const = 0;
+	/** Gets the tree root item of the Navigation Tool */
+	virtual FNavigationToolViewModelWeakPtr GetTreeRoot() const = 0;
 
 	/**
 	 * Finds the registered item that has the given Id
 	 * @returns The item with the given Id, or null if the item does not exist or was not registered to the Navigation Tool
 	 */
-	virtual FNavigationToolItemPtr FindItem(const FNavigationToolItemId& InItemId) const = 0;
+	virtual FNavigationToolViewModelPtr FindItem(const FNavigationToolItemId& InItemId) const = 0;
 
 	/**
-	 * Attempts to find all items associated with a Sequencer view model
-	 * @param InViewModel The Sequencer view model to find the items of
+	 * Attempts to find all items associated with a Sequencer outliner view model
+	 * @param InSequencerOutlinerViewModel The Sequencer view model to find the items for
 	 * @return The list of Navigation Tool items found
 	 */
-	virtual TArray<FNavigationToolItemPtr> TryFindItems(const UE::Sequencer::FViewModelPtr& InViewModel) const = 0;
+	virtual TArray<FNavigationToolViewModelPtr> TryFindItems(const Sequencer::FViewModelPtr& InSequencerOutlinerViewModel) const = 0;
 
 	/**
  	 * Adds or Removes the Ignore Notify Flags to prevent certain actions from automatically happening when they're triggered
@@ -186,14 +187,15 @@ public:
 	virtual void OnSequencerSelectionChanged() = 0;
 
 	/** Returns the currently selected items in the most recent outliner view (since this list can vary between outliner views) */
-	virtual TArray<FNavigationToolItemPtr> GetSelectedItems(const bool bInNormalizeToTopLevelSelections = false) const = 0;
+	virtual TArray<FNavigationToolViewModelWeakPtr> GetSelectedItems(const bool bInNormalizeToTopLevelSelections = false) const = 0;
 
 	/**
 	 * Selects the given Items on all Navigation Tool Views
-	 * @param InItems the items to select
+	 * @param InWeakItems the items to select
 	 * @param InFlags how the items should be selected (appended, notify of selections, etc)
 	 */
-	virtual void SelectItems(const TArray<FNavigationToolItemPtr>& InItems, const ENavigationToolItemSelectionFlags InFlags = ENavigationToolItemSelectionFlags::SignalSelectionChange) const = 0;
+	virtual void SelectItems(const TArray<FNavigationToolViewModelWeakPtr>& InWeakItems
+		, const ENavigationToolItemSelectionFlags InFlags = ENavigationToolItemSelectionFlags::SignalSelectionChange) const = 0;
 
 	/**
 	 * Clears the Item Selection from all Navigation Tool Views
@@ -202,10 +204,10 @@ public:
 	virtual void ClearItemSelection(const bool bInSignalSelectionChange) const = 0;
 
 	/** Pairs the Item with the given color name, overriding the inherited color if different */
-	virtual void SetItemColor(const FNavigationToolItemPtr& InItem, const FColor& InColor) = 0;
+	virtual void SetItemColor(const FNavigationToolViewModelPtr& InItem, const FColor& InColor) = 0;
 
 	/** Removes the Color pairing of the given Item (can still have an inherited color though) */
-	virtual void RemoveItemColor(const FNavigationToolItemPtr& InItem) = 0;
+	virtual void RemoveItemColor(const FNavigationToolViewModelPtr& InItem) = 0;
 
 	/**
 	 * Gets the color related to the Item
@@ -213,14 +215,14 @@ public:
 	 * @param bRecurseParent whether to get the color of the parent (recursively) if the given item does not have a color by itself
 	 * @returns the matching color pair or unset if item is invalid or no color could be found.
 	 */
-	virtual TOptional<FColor> FindItemColor(const FNavigationToolItemPtr& InItem, bool bRecurseParent = true) const = 0;
+	virtual TOptional<FColor> FindItemColor(const FNavigationToolViewModelPtr& InItem, bool bRecurseParent = true) const = 0;
 
 	/**
 	 * Instantiates a new item action without adding it to the Pending Actions Queue.
 	 * This should only be used directly when planning to enqueue multiple actions.
 	 * @see FNavigationTool::EnqueueItemActions
 	 */
-	template<typename InItemActionType, typename = typename TEnableIf<TIsDerivedFrom<InItemActionType, INavigationToolAction>::IsDerived>::Type, typename ...InArgTypes>
+	template<typename InItemActionType, typename = typename TEnableIf<TIsDerivedFrom<InItemActionType, INavigationToolItemAction>::IsDerived>::Type, typename ...InArgTypes>
 	TSharedRef<InItemActionType> NewItemAction(InArgTypes&&... InArgs)
 	{
 		return MakeShared<InItemActionType>(Forward<InArgTypes>(InArgs)...);
@@ -231,20 +233,20 @@ public:
 	 * Ideal for when dealing with a single action.
 	 * For multiple actions use FNavigationTool::EnqueueItemActions.
 	 */
-	template<typename InItemActionType, typename = typename TEnableIf<TIsDerivedFrom<InItemActionType, INavigationToolAction>::IsDerived>::Type, typename ...InArgTypes>
+	template<typename InItemActionType, typename = typename TEnableIf<TIsDerivedFrom<InItemActionType, INavigationToolItemAction>::IsDerived>::Type, typename ...InArgTypes>
 	void EnqueueItemAction(InArgTypes&&... InArgs)
 	{
 		EnqueueItemActions({ NewItemAction<InItemActionType>(Forward<InArgTypes>(InArgs)...) });
 	}
 
 	/** Adds the given actions to the Pending Action Queue */
-	virtual void EnqueueItemActions(TArray<TSharedPtr<INavigationToolAction>>&& InItemActions) noexcept = 0;
+	virtual void EnqueueItemActions(TArray<TSharedPtr<INavigationToolItemAction>>&& InItemActions) noexcept = 0;
 
 	/** Called when an Item has been Renamed */
-	virtual void NotifyToolItemRenamed(const FNavigationToolItemPtr& InItem) = 0;
+	virtual void NotifyToolItemRenamed(const FNavigationToolViewModelPtr& InItem) = 0;
 
 	/** Called when an Item has been deleted */
-	virtual void NotifyToolItemDeleted(const FNavigationToolItemPtr& InItem) = 0;
+	virtual void NotifyToolItemDeleted(const FNavigationToolViewModelPtr& InItem) = 0;
 
 protected:
 	virtual const FNavigationToolItemProxyRegistry& GetItemProxyRegistry() const = 0;
diff --git a/Engine/Plugins/Experimental/SequenceNavigator/Source/SequenceNavigator/Public/INavigationToolView.h b/Engine/Plugins/Experimental/SequenceNavigator/Source/SequenceNavigator/Public/INavigationToolView.h
--- a/Engine/Plugins/Experimental/SequenceNavigator/Source/SequenceNavigator/Public/INavigationToolView.h
+++ b/Engine/Plugins/Experimental/SequenceNavigator/Source/SequenceNavigator/Public/INavigationToolView.h
@@ -49,65 +49,73 @@ public:
 
 	virtual TSharedRef<SWidget> GetColumnMenuContent(const FName InColumnId) = 0;
 
-	virtual void GetChildrenOfItem(const FNavigationToolItemPtr InItem, TArray<FNavigationToolItemPtr>& OutChildren) const = 0;
+	virtual void GetChildrenOfItem(const FNavigationToolViewModelWeakPtr InWeakItem
+		, TArray<FNavigationToolViewModelWeakPtr>& OutWeakChildren) const = 0;
 
 	/**
-	 * Gets the Children of a given Item. Can recurse if the immediate child is hidden (the children of these hidden items should still be given the opportunity to show up)
-	 * @param InItem the Item to get the children of
-	 * @param OutChildren the visible children in the give view mode.
-	 * @param InViewMode the view mode(s) that the children should support to be added to OutChildren
-	 * @param InRecursionDisallowedItems the items where recursion should not be performed
+	 * Gets the children of a given item. Can recurse if the immediate child is hidden (the children of these hidden items should still be given the opportunity to show up)
+	 * @param InWeakItem The item to get the children of
+	 * @param OutWeakChildren The visible children in the give view mode
+	 * @param InViewMode The view mode(s) that the children should support to be added to OutChildren
+	 * @param InWeakRecursionDisallowedItems The items where recursion should not be performed
 	 */
-	virtual void GetChildrenOfItem(const FNavigationToolItemPtr& InItem
-		, TArray<FNavigationToolItemPtr>& OutChildren
-		, ENavigationToolItemViewMode InViewMode
-		, const TSet<FNavigationToolItemPtr>& InRecursionDisallowedItems) const = 0;
+	virtual void GetChildrenOfItem(const FNavigationToolViewModelWeakPtr InWeakItem
+		, TArray<FNavigationToolViewModelWeakPtr>& OutWeakChildren
+		, const ENavigationToolItemViewMode InViewMode
+		, const TSet<FNavigationToolViewModelWeakPtr>& InWeakRecursionDisallowedItems) const = 0;
 
-	/** Whether the given Item is explicitly marked as Read-only in the Navigation Tool */
-	virtual bool IsItemReadOnly(const FNavigationToolItemPtr& InItem) const = 0;
+	/** Whether the given item is explicitly marked as Read-only in the Navigation Tool */
+	virtual bool IsItemReadOnly(const FNavigationToolViewModelWeakPtr& InWeakItem) const = 0;
 
 	/** Selection State */
 
-	/** Whether the given Item can be selected in this Navigation Tool View / Widget */
-	virtual bool CanSelectItem(const FNavigationToolItemPtr& InItem) const = 0;
+	/** Whether the given item can be selected in this Navigation Tool View / Widget */
+	virtual bool CanSelectItem(const FNavigationToolViewModelWeakPtr& InWeakItem) const = 0;
 
-	/** Selects the Item in this Navigation Tool View */
-	virtual void SelectItems(TArray<FNavigationToolItemPtr> InItems, const ENavigationToolItemSelectionFlags InFlags = ENavigationToolItemSelectionFlags::SignalSelectionChange) = 0;
+	/** Selects the item in this Navigation Tool View */
+	virtual void SelectItems(TArray<FNavigationToolViewModelWeakPtr> InWeakItems
+		, const ENavigationToolItemSelectionFlags InFlags = ENavigationToolItemSelectionFlags::SignalSelectionChange) = 0;
 
-	/** Clears the currently Selected Items in the Navigation Tool View */
-	virtual void ClearItemSelection(bool bSignalSelectionChange = true) = 0;
+	/** Clears the currently selected items in the Navigation Tool View */
+	virtual void ClearItemSelection(const bool bInSignalSelectionChange = true) = 0;
 
 	/** Whether the current item is selected in this Navigation Tool View */
-	virtual bool IsItemSelected(const FNavigationToolItemPtr& InItem) const = 0;
+	virtual bool IsItemSelected(const FNavigationToolViewModelWeakPtr& InWeakItem) const = 0;
 
-	/** Gets the Currently Selected Items in the Tree View */
-	virtual TArray<FNavigationToolItemPtr> GetSelectedItems() const = 0;
+	/** Gets the currently selected items in the Tree View */
+	virtual TArray<FNavigationToolViewModelWeakPtr> GetSelectedItems() const = 0;
 
 	/** Drag Drop */
-
-	/** Called a Drag is attempted for the selected items in view */
-	virtual FReply OnDragDetected(const FGeometry& InGeometry, const FPointerEvent& InMouseEvent, const FNavigationToolItemPtr InTargetItem) = 0;
-
-	/** Processes the Drag and Drop for the given Item if valid, else it will process it for the Root Item */
-	virtual FReply OnDrop(const FDragDropEvent& InDragDropEvent, EItemDropZone InDropZone, const FNavigationToolItemPtr InTargetItem) = 0;
-
-	/** Determines whether the Drag and Drop can be processed by the given item */
+#if 1==2
+	/** Called a drag is attempted for the selected items in view */
+	virtual FReply OnDragDetected(const FGeometry& InGeometry
+		, const FPointerEvent& InPointerEvent
+		, const FNavigationToolViewModelWeakPtr InWeakTargetItem) = 0;
+
+	/** Processes the drag and drop for the given item if valid, else it will process it for the root item */
+	virtual FReply OnDrop(const FDragDropEvent& InDragDropEvent
+		, const EItemDropZone InDropZone
+		, const FNavigationToolViewModelWeakPtr InWeakTargetItem) = 0;
+
+	/** Determines whether the drag and drop can be processed by the given item */
 	virtual TOptional<EItemDropZone> OnCanDrop(const FDragDropEvent& InDragDropEvent
-		, EItemDropZone InDropZone
-		, const FNavigationToolItemPtr InTargetItem) const = 0;
-
+		, const EItemDropZone InDropZone
+		, const FNavigationToolViewModelWeakPtr InWeakTargetItem) const = 0;
+#endif
 	/** Expansion State */
 
-	virtual bool IsItemExpanded(const FNavigationToolItemPtr& InItem, const bool bInUseFilter = true) const = 0;
+	virtual bool IsItemExpanded(const FNavigationToolViewModelWeakPtr& InWeakItem
+		, const bool bInUseFilter = true) const = 0;
 
-	/** Expands / Collapses the given Item */
-	virtual void SetItemExpansion(const FNavigationToolItemPtr& InItem, const bool bInExpand, const bool bInUseFilter = true) = 0;
+	/** Expands / Collapses the given item */
+	virtual void SetItemExpansion(const FNavigationToolViewModelWeakPtr& InWeakItem
+		, const bool bInExpand, const bool bInUseFilter = true) = 0;
 
-	/** Expands / Collapses the given Item and its children recursively */
-	virtual void SetItemExpansionRecursive(const FNavigationToolItemPtr InItem, const bool bInExpand) = 0;
+	/** Expands / Collapses the given item and its children recursively */
+	virtual void SetItemExpansionRecursive(const FNavigationToolViewModelWeakPtr InWeakItem, const bool bInExpand) = 0;
 
 	/** Change the expansion state of the parents (recursively) of the given item */
-	virtual void SetParentItemExpansions(const FNavigationToolItemPtr& InItem, const bool bInExpand) = 0;
+	virtual void SetParentItemExpansions(const FNavigationToolViewModelWeakPtr& InWeakItem, const bool bInExpand) = 0;
 
 	/** @return True if all items in the tool tree view can be expanded */
 	virtual bool CanExpandAll() const = 0;
diff --git a/Engine/Plugins/Experimental/SequenceNavigator/Source/SequenceNavigator/Public/ItemActions/NavigationToolAddItem.h b/Engine/Plugins/Experimental/SequenceNavigator/Source/SequenceNavigator/Public/ItemActions/NavigationToolAddItem.h
--- a/Engine/Plugins/Experimental/SequenceNavigator/Source/SequenceNavigator/Public/ItemActions/NavigationToolAddItem.h
+++ b/Engine/Plugins/Experimental/SequenceNavigator/Source/SequenceNavigator/Public/ItemActions/NavigationToolAddItem.h
@@ -2,8 +2,8 @@
 
 #pragma once
 
+#include "INavigationToolItemAction.h"
 #include "Items/NavigationToolItemParameters.h"
-#include "NavigationToolItemAction.h"
 
 #define UE_API SEQUENCENAVIGATOR_API
 
@@ -16,12 +16,15 @@ class FNavigationTool;
  * Item action responsible for adding an item to the tree under a given optional parent.
  * If Parent is null, it is added as a Top Level Item.
  */
-class FNavigationToolAddItem : public INavigationToolAction
+class FNavigationToolAddItem
+		: public INavigationToolItemAction
 {
 public:
-	UE_NAVIGATIONTOOL_INHERITS(FNavigationToolAddItem, INavigationToolAction);
+	UE_SEQUENCER_DECLARE_CASTABLE_API(UE_API, FNavigationToolAddItem
+		, INavigationToolItemAction)
 
 	UE_API FNavigationToolAddItem(const FNavigationToolAddItemParams& InAddItemParams);
+	virtual ~FNavigationToolAddItem() = default;
 
 	//~ Begin INavigationToolAction
 	UE_API virtual bool ShouldTransact() const override;
diff --git a/Engine/Plugins/Experimental/SequenceNavigator/Source/SequenceNavigator/Public/ItemActions/NavigationToolRemoveItem.h b/Engine/Plugins/Experimental/SequenceNavigator/Source/SequenceNavigator/Public/ItemActions/NavigationToolRemoveItem.h
--- a/Engine/Plugins/Experimental/SequenceNavigator/Source/SequenceNavigator/Public/ItemActions/NavigationToolRemoveItem.h
+++ b/Engine/Plugins/Experimental/SequenceNavigator/Source/SequenceNavigator/Public/ItemActions/NavigationToolRemoveItem.h
@@ -2,7 +2,7 @@
 
 #pragma once
 
-#include "NavigationToolItemAction.h"
+#include "INavigationToolItemAction.h"
 #include "Items/NavigationToolItemParameters.h"
 
 #define UE_API SEQUENCENAVIGATOR_API
@@ -15,16 +15,19 @@ class FNavigationTool;
 /**
  * Item Action responsible for removing/unregistering items from the tree
  */
-class FNavigationToolRemoveItem : public INavigationToolAction
+class FNavigationToolRemoveItem
+	: public INavigationToolItemAction
 {
 public:
-	UE_NAVIGATIONTOOL_INHERITS(FNavigationToolRemoveItem, INavigationToolAction);
+	UE_SEQUENCER_DECLARE_CASTABLE_API(UE_API, FNavigationToolRemoveItem
+		, INavigationToolItemAction)
 
-	UE_API FNavigationToolRemoveItem(const FNavigationToolRemoveItemParams& InRemoveItemParams);
+	FNavigationToolRemoveItem(const FNavigationToolRemoveItemParams& InRemoveItemParams);
+	virtual ~FNavigationToolRemoveItem() = default;
 
 	//~ Begin INavigationToolAction
-	UE_API virtual void Execute(FNavigationTool& InTool) override;
-	UE_API virtual void OnObjectsReplaced(const TMap<UObject*, UObject*>& InReplacementMap, bool bInRecursive) override;
+	virtual void Execute(FNavigationTool& InTool) override;
+	virtual void OnObjectsReplaced(const TMap<UObject*, UObject*>& InReplacementMap, const bool bInRecursive) override;
 	//~ End INavigationToolAction
 
 protected:
diff --git a/Engine/Plugins/Experimental/SequenceNavigator/Source/SequenceNavigator/Public/ItemProxies/Factories/NavigationToolItemProxyDefaultFactory.h b/Engine/Plugins/Experimental/SequenceNavigator/Source/SequenceNavigator/Public/ItemProxies/Factories/NavigationToolItemProxyDefaultFactory.h
--- a/Engine/Plugins/Experimental/SequenceNavigator/Source/SequenceNavigator/Public/ItemProxies/Factories/NavigationToolItemProxyDefaultFactory.h
+++ b/Engine/Plugins/Experimental/SequenceNavigator/Source/SequenceNavigator/Public/ItemProxies/Factories/NavigationToolItemProxyDefaultFactory.h
@@ -4,7 +4,6 @@
 
 #include "ItemProxies/INavigationToolItemProxyFactory.h"
 #include "Items/NavigationToolItemProxy.h"
-#include "NavigationToolItemType.h"
 
 namespace UE::SequenceNavigator
 {
@@ -20,7 +19,7 @@ class TNavigationToolItemProxyDefaultFactoryBase : public INavigationToolItemPro
 public:
 	virtual FName GetItemProxyTypeName() const override
 	{
-		return TNavigationToolItemType<InItemProxyType>::GetTypeId().ToName();
+		return Sequencer::TAutoRegisterViewModelTypeID<InItemProxyType>::GetTypeTable().GetTypeName();
 	}
 
 protected:
@@ -38,7 +37,7 @@ template<typename InItemProxyType, uint32 InItemProxyPriority = 0
 class TNavigationToolItemProxyDefaultFactory : public TNavigationToolItemProxyDefaultFactoryBase<InItemProxyType, InItemProxyPriority>
 {
 public:
-	virtual TSharedPtr<FNavigationToolItemProxy> CreateItemProxy(INavigationTool& InTool, const FNavigationToolItemPtr& InParentItem) override
+	virtual TSharedPtr<FNavigationToolItemProxy> CreateItemProxy(INavigationTool& InTool, const FNavigationToolViewModelPtr& InParentItem) override
 	{
 		return this->DefaultCreateItemProxy(InTool, InParentItem);
 	}
diff --git a/Engine/Plugins/Experimental/SequenceNavigator/Source/SequenceNavigator/Public/ItemProxies/INavigationToolItemProxyFactory.h b/Engine/Plugins/Experimental/SequenceNavigator/Source/SequenceNavigator/Public/ItemProxies/INavigationToolItemProxyFactory.h
--- a/Engine/Plugins/Experimental/SequenceNavigator/Source/SequenceNavigator/Public/ItemProxies/INavigationToolItemProxyFactory.h
+++ b/Engine/Plugins/Experimental/SequenceNavigator/Source/SequenceNavigator/Public/ItemProxies/INavigationToolItemProxyFactory.h
@@ -2,6 +2,7 @@
 
 #pragma once
 
+#include "Items/NavigationToolItem.h"
 #include "NavigationToolDefines.h"
 #include "Templates/SharedPointer.h"
 
@@ -10,7 +11,6 @@ class FName;
 namespace UE::SequenceNavigator
 {
 
-class FNavigationToolItem;
 class FNavigationToolItemProxy;
 class INavigationTool;
 
@@ -27,7 +27,8 @@ public:
 	virtual FName GetItemProxyTypeName() const = 0;
 
 	/** Returns a newly created instance of the Relevant Item Proxy if successful */
-	virtual TSharedPtr<FNavigationToolItemProxy> CreateItemProxy(INavigationTool& InTool, const FNavigationToolItemPtr& InParentItem) = 0;
+	virtual TSharedPtr<FNavigationToolItemProxy> CreateItemProxy(INavigationTool& InTool
+		, const FNavigationToolViewModelPtr& InParentItem) = 0;
 };
 
 } // namespace UE::SequenceNavigator
diff --git a/Engine/Plugins/Experimental/SequenceNavigator/Source/SequenceNavigator/Public/ItemProxies/NavigationToolComponentProxy.h b/Engine/Plugins/Experimental/SequenceNavigator/Source/SequenceNavigator/Public/ItemProxies/NavigationToolComponentProxy.h
--- a/Engine/Plugins/Experimental/SequenceNavigator/Source/SequenceNavigator/Public/ItemProxies/NavigationToolComponentProxy.h
+++ b/Engine/Plugins/Experimental/SequenceNavigator/Source/SequenceNavigator/Public/ItemProxies/NavigationToolComponentProxy.h
@@ -9,22 +9,25 @@
 namespace UE::SequenceNavigator
 {
 
-class FNavigationToolComponentProxy : public FNavigationToolItemProxy
+class FNavigationToolComponentProxy
+	: public FNavigationToolItemProxy
 {
 public:
-	UE_NAVIGATIONTOOL_INHERITS_WITH_SUPER(FNavigationToolComponentProxy, FNavigationToolItemProxy)
+	UE_SEQUENCER_DECLARE_CASTABLE_API(UE_API, FNavigationToolComponentProxy
+		, FNavigationToolItemProxy)
+
+	UE_API FNavigationToolComponentProxy(INavigationTool& InTool, const FNavigationToolViewModelPtr& InParentItem);
 
-	UE_API FNavigationToolComponentProxy(INavigationTool& InTool, const FNavigationToolItemPtr& InParentItem);
-	
 	//~ Begin INavigationToolItem
 	UE_API virtual FText GetDisplayName() const override;
 	UE_API virtual FSlateIcon GetIcon() const override;
 	UE_API virtual FText GetIconTooltipText() const override;
 	UE_API virtual ENavigationToolItemViewMode GetSupportedViewModes(const INavigationToolView& InToolView) const override;
 	//~ End INavigationToolItem
-	
+
 	//~ Begin FNavigationToolItemProxy
-	UE_API virtual void GetProxiedItems(const TSharedRef<INavigationToolItem>& InParent, TArray<FNavigationToolItemPtr>& OutChildren, bool bInRecursive) override;
+	UE_API virtual void GetProxiedItems(const FNavigationToolViewModelPtr& InParent
+		, TArray<FNavigationToolViewModelWeakPtr>& OutWeakChildren, const bool bInRecursive) override;
 	//~ End FNavigationToolItemProxy
 };
 
diff --git a/Engine/Plugins/Experimental/SequenceNavigator/Source/SequenceNavigator/Public/ItemProxies/NavigationToolItemProxyRegistry.h b/Engine/Plugins/Experimental/SequenceNavigator/Source/SequenceNavigator/Public/ItemProxies/NavigationToolItemProxyRegistry.h
--- a/Engine/Plugins/Experimental/SequenceNavigator/Source/SequenceNavigator/Public/ItemProxies/NavigationToolItemProxyRegistry.h
+++ b/Engine/Plugins/Experimental/SequenceNavigator/Source/SequenceNavigator/Public/ItemProxies/NavigationToolItemProxyRegistry.h
@@ -6,7 +6,6 @@
 #include "Containers/Map.h"
 #include "Containers/Set.h"
 #include "Factories/NavigationToolItemProxyDefaultFactory.h"
-#include "NavigationToolItemType.h"
 #include "Templates/SharedPointer.h"
 #include "UObject/NameTypes.h"
 
@@ -44,7 +43,7 @@ public:
 	template<typename InItemProxyType, typename = typename TEnableIf<TIsDerivedFrom<InItemProxyType, FNavigationToolItemProxy>::IsDerived>::Type>
 	void UnregisterItemProxyFactory()
 	{
-		ItemProxyFactories.Remove(TNavigationToolItemType<InItemProxyType>::GetTypeId().ToName());
+		ItemProxyFactories.Remove(Sequencer::TAutoRegisterViewModelTypeID<InItemProxyType>::GetTypeTable().GetTypeName());
 	}
 
 	void UnregisterItemProxyFactory(FName InItemProxyTypeName)
@@ -59,7 +58,7 @@ public:
 	}
 
 	/** Gets the Item Proxy Factory for the given Item Proxy Type Name. Returns nullptr if not found */
-	INavigationToolItemProxyFactory* GetItemProxyFactory(FName InItemProxyTypeName) const
+	INavigationToolItemProxyFactory* GetItemProxyFactory(const FName InItemProxyTypeName) const
 	{
 		if (const TSharedRef<INavigationToolItemProxyFactory>* const FoundFactory = ItemProxyFactories.Find(InItemProxyTypeName))
 		{
@@ -72,7 +71,7 @@ public:
 	template<typename InItemProxyType, typename = typename TEnableIf<TIsDerivedFrom<InItemProxyType, FNavigationToolItemProxy>::IsDerived>::Type>
 	INavigationToolItemProxyFactory* GetItemProxyFactory() const
 	{
-		return GetItemProxyFactory(TNavigationToolItemType<InItemProxyType>::GetTypeId().ToName());
+		return GetItemProxyFactory(InItemProxyType::ID.GetTypeTable()->GetTypeName());
 	}
 
 	/** Gets all the Item Proxy Type Names that exist in this registry */
diff --git a/Engine/Plugins/Experimental/SequenceNavigator/Source/SequenceNavigator/Public/Items/INavigationToolItem.h b/Engine/Plugins/Experimental/SequenceNavigator/Source/SequenceNavigator/Public/Items/INavigationToolItem.h
--- a/Engine/Plugins/Experimental/SequenceNavigator/Source/SequenceNavigator/Public/Items/INavigationToolItem.h
+++ b/Engine/Plugins/Experimental/SequenceNavigator/Source/SequenceNavigator/Public/Items/INavigationToolItem.h
@@ -4,16 +4,17 @@
 
 #include "Delegates/Delegate.h"
 #include "EditorSequenceNavigationDefs.h"
+#include "MVVM/ViewModelTypeID.h"
 #include "NavigationToolDefines.h"
-#include "NavigationToolItemId.h"
-#include "NavigationToolItemType.h"
+#include "NavigationToolSettings.h"
+
+#define UE_API SEQUENCENAVIGATOR_API
 
 class FDragDropEvent;
 class FReply;
 class SWidget;
 class UObject;
 enum class EItemDropZone;
-enum class ENavigationToolItemViewMode : uint8;
 struct FLinearColor;
 struct FSlateBrush;
 struct FSlateColor;
@@ -33,16 +34,20 @@ struct FNavigationToolItemId;
 struct FNavigationToolRemoveItemParams;
 
 /*
- * An Navigation Tool Item is the class that represents a Single Element (i.e. Node) in the Navigation Tool Tree.
- * This can be an Item that represents an Object (e.g. Actor,Component) or a Folder, or something else.
+ * An Navigation Tool Item is the class that represents a single element (i.e. node) in the Navigation Tool tree.
+ * This can be an item that represents an object (e.g. actor, component) or a folder, or something else.
 */
-class INavigationToolItem : public INavigationToolItemTypeCastable, public TSharedFromThis<INavigationToolItem>
+class INavigationToolItem
 {
-protected:
 	using IndexType = TArray<INavigationToolItem>::SizeType;
 
 public:
-	UE_NAVIGATIONTOOL_INHERITS(INavigationToolItem, INavigationToolItemTypeCastable);
+	UE_SEQUENCER_DECLARE_VIEW_MODEL_TYPE_ID_API(UE_API, INavigationToolItem)
+
+	DECLARE_MULTICAST_DELEGATE_TwoParams(FOnExpansionChanged, const TSharedPtr<INavigationToolView>&, bool);
+	DECLARE_MULTICAST_DELEGATE_TwoParams(FOnRenameAction, ENavigationToolRenameAction, const TSharedPtr<INavigationToolView>&);
+
+	virtual ~INavigationToolItem() = default;
 
 	/** Determines whether the Item properties are in a valid state */
 	virtual bool IsItemValid() const = 0;
@@ -86,18 +91,8 @@ public:
 	/** Refreshes what the Parent and Children are of this Item. (not recursive!) */
 	virtual void RefreshChildren() = 0;
 
-	/** Resets both the Cached Visible Children and Children (before doing so, sets all child's parents to null) */
-	virtual void ResetChildren()
-	{
-		for (const FNavigationToolItemPtr& Item : GetChildren())
-		{
-			if (Item.IsValid())
-			{
-				Item->SetParent(nullptr);
-			}
-		}
-		GetChildrenMutable().Reset();	
-	}
+	/** Resets both the cached visible children and children (before doing so, sets all child's parents to null) */
+	virtual void ResetChildren() = 0;
 
 	/**
 	 * Determines whether this item can be sorted by the Navigation Tool or not.
@@ -107,7 +102,7 @@ public:
 	virtual bool ShouldSort() const = 0;
 
 	/** Determines whether the given Child is supported and can be added under this Item */
-	virtual bool CanAddChild(const FNavigationToolItemPtr& InChild) const = 0;
+	virtual bool CanAddChild(const FNavigationToolViewModelPtr& InChild) const = 0;
 
 	/**
 	 * Adds another Child under this Item if such Item is supported.
@@ -121,301 +116,69 @@ public:
 	 */
 	virtual bool RemoveChild(const FNavigationToolRemoveItemParams& InRemoveItemParams) = 0;
 
-	/** Figures out the Children of this Item. This is only relevant for items that do have that functionality (e.g. Components or Actors) */
-	virtual void FindChildren(TArray<FNavigationToolItemPtr>& OutChildren, const bool bInRecursive) = 0;
-
-	virtual void FindValidChildren(TArray<FNavigationToolItemPtr>& OutChildren, const bool bInRecursive)
-	{
-		FindChildren(OutChildren, bInRecursive);
+	/** Finds this items children. This is only relevant for items that do have that functionality (e.g. components or actors) */
+	virtual void FindChildren(TArray<FNavigationToolViewModelWeakPtr>& OutWeakChildren, const bool bInRecursive) = 0;
 
-		OutChildren.RemoveAll(
-			[](const FNavigationToolItemPtr& InItem)
-			{
-				return !InItem.IsValid() || !InItem->IsAllowedInTool();
-			});
-	}
+	virtual void FindValidChildren(TArray<FNavigationToolViewModelWeakPtr>& OutWeakChildren, const bool bInRecursive) = 0;
 
 	/** Gets the Item Proxies for this Item (e.g. Component Item that represent Primitives add in a Material Proxy to display) */
 	virtual void GetItemProxies(TArray<TSharedPtr<FNavigationToolItemProxy>>& OutItemProxies) {}
 
 	/*
-	 * Tries to Find the First Path of Descendants (not including self) that lead to a given Item in the Set.
-	 * The last item is the Item in the Set that was found so path might be A/B/C/.../ItemInSet  where A is Child of This.
-	 * Returns Empty array if failed.
+	 * Tries to find the first path of descendants (not including self) that lead to a given item in the set.
+	 * The last item is the item in the set that was found, so the path might be "A/B/C/.../ItemInSet", where A is the child of This.
+	 * Returns The empty array if failed.
 	*/
-	virtual TArray<FNavigationToolItemPtr> FindPath(const TArray<FNavigationToolItemPtr>& Items) const
-	{
-		const TSharedPtr<const INavigationToolItem> This = SharedThis(this);
-		TArray<FNavigationToolItemPtr> Path;
-
-		for (FNavigationToolItemPtr Item : Items)
-		{
-			Path.Reset();
-			FNavigationToolItemPtr CurrentItem = Item;
-
-			while (CurrentItem.IsValid())
-			{
-				if (This == CurrentItem)
-				{
-					Algo::Reverse(Path);
-					return Path;
-				}
-
-				Path.Add(CurrentItem);
-				CurrentItem = CurrentItem->GetParent();
-			}
-		}
-
-		return TArray<FNavigationToolItemPtr>();
-	}
+	virtual TArray<FNavigationToolViewModelPtr> FindPath(const TArray<FNavigationToolViewModelWeakPtr>& InWeakItems) const = 0;
 
 	/** Returns the path to this item in the tree. Ex. "RootId,ChildId,ChildId" */
-	FString GetFullPath() const
-	{
-		FString OutPath;
-
-		FNavigationToolItemPtr CurrentParent = GetParent();
-		while (CurrentParent.IsValid())
-		{
-			if (OutPath.IsEmpty())
-			{
-				OutPath = CurrentParent->GetItemId().GetStringId();
-			}
-			else
-			{
-				OutPath = CurrentParent->GetItemId().GetStringId() + TEXT(",") + OutPath;
-			}
-
-			CurrentParent = CurrentParent->GetParent();
-		}
-
-		return OutPath;
-	}
+	virtual FString GetFullPath() const = 0;
 
-	/** Gets the current Child Items of this Item */
-	virtual const TArray<FNavigationToolItemPtr>& GetChildren() const = 0;
+	/** Gets the current Child Items of this item */
+	virtual const TArray<FNavigationToolViewModelWeakPtr>& GetChildren() const = 0;
 	
 	/** Gets the current Child Items of this Item */
-	virtual TArray<FNavigationToolItemPtr>& GetChildrenMutable() = 0;
+	virtual TArray<FNavigationToolViewModelWeakPtr>& GetChildrenMutable() = 0;
 
 	/**
 	 * Gets the Index that the given Child Item is at.
 	 * NOTE: This includes HIDDEN items as Item Visibility is relative to each Navigation Tool View.
 	 * To consider only Visible Items, use FNavigationToolView::GetVisibleChildAt.
 	 */
-	virtual IndexType GetChildIndex(const FNavigationToolItemPtr& ChildItem) const
-	{
-		return GetChildren().Find(ChildItem);
-	}
+	virtual IndexType GetChildIndex(const FNavigationToolViewModelWeakPtr& InWeakChildItem) const = 0;
 
 	/*
 	 * Gets the Child Item at the given Index.
 	 * NOTE: This includes HIDDEN items as Item Visibility is relative to each Navigation Tool View.
 	 * To consider only Visible Items, use FNavigationToolView::GetVisibleChildAt.
 	 */
-	virtual FNavigationToolItemPtr GetChildAt(const IndexType Index) const
-	{
-		const TArray<FNavigationToolItemPtr>& ChildItems = GetChildren();
-		if (ChildItems.IsValidIndex(Index))
-		{
-			return ChildItems[Index];
-		}
-		return nullptr;
-	};
-
-	/** Gets the Parent of this Item. Should only be null prior to registering it in Navigation Tool or if its Root Item */
-	virtual FNavigationToolItemPtr GetParent() const = 0;
-
-	/** Gets the list of all parents up the tree from this Item */
-	virtual TSet<FNavigationToolItemPtr> GetParents(const bool bInIncludeRoot = false) const
-	{
-		TSet<FNavigationToolItemPtr> OutParents;
-
-		FNavigationToolItemPtr CurrentParent = GetParent();
-		while (CurrentParent.IsValid()
-			&& (bInIncludeRoot || CurrentParent->GetItemId() != FNavigationToolItemId::RootId))
-		{
-			OutParents.Add(CurrentParent);
-
-			CurrentParent = CurrentParent->GetParent();
-		}
-
-		return OutParents;
-	}
-
-	/** Sets the Parent. Note that the Parent must've already have this instance as a child (check is done) */
-	virtual void SetParent(FNavigationToolItemPtr InParent) = 0;
-
-	template<typename InItemType>
-	InItemType* FindAncestorOfType(const bool bInIncludeThis = false) const
-	{
-		if (bInIncludeThis)
-		{
-			if (InItemType* const CastedParent = const_cast<InItemType*>(CastTo<const InItemType>()))
-			{
-				return CastedParent;
-			}
-		}
-
-		FNavigationToolItemPtr CurrentParent = GetParent();
-
-		while (CurrentParent.IsValid())
-		{
-			if (InItemType* const CastedParent = CurrentParent->CastTo<InItemType>())
-			{
-				return CastedParent;
-			}
-
-			CurrentParent = CurrentParent->GetParent();
-		}
-
-		return nullptr;
-	}
-
-	template<typename InItemType>
-	TArray<InItemType*> FindAncestorsOfType(const bool bInIncludeThis = false) const
-	{
-		TArray<InItemType*> OutAncestors;
-
-		if (bInIncludeThis)
-		{
-			if (InItemType* const CastedParent = const_cast<InItemType*>(CastTo<const InItemType>()))
-			{
-				OutAncestors.Add(CastedParent);
-			}
-		}
-
-		FNavigationToolItemPtr CurrentParent = GetParent();
-
-		while (CurrentParent.IsValid())
-		{
-			if (InItemType* const CastedParent = CurrentParent->CastTo<InItemType>())
-			{
-				OutAncestors.Add(CastedParent);
-			}
-
-			CurrentParent = CurrentParent->GetParent();
-		}
-
-		return OutAncestors;
-	}
-
-	TArray<FNavigationToolItemPtr> GetDescendants(const bool bInIncludeThis = false)
-	{
-		TArray<FNavigationToolItemPtr> OutDescendants;
-
-		if (bInIncludeThis)
-		{
-			OutDescendants.Add(AsShared());
-		}
-
-		for (const FNavigationToolItemPtr& Child : GetChildren())
-		{
-			OutDescendants.Add(Child);
-
-			if (Child.IsValid())
-			{
-				OutDescendants.Append(Child->GetDescendants());
-			}
-		}
-
-		return OutDescendants;
-	}
-
-	template<typename InItemType>
-	TArray<InItemType*> GetDescendantsOfType(const bool bInIncludeThis = false) const
-	{
-		TArray<InItemType*> OutDescendants;
-
-		if (bInIncludeThis)
-		{
-			if (InItemType* const CastedChild = const_cast<InItemType*>(CastTo<const InItemType>()))
-			{
-				OutDescendants.Add(CastedChild);
-			}
-		}
-
-		// Process each child of the current item recursively
-		for (const FNavigationToolItemPtr& Child : GetChildren())
-		{
-			// If the child is of the requested type, cast it and add to the result
-			if (InItemType* CastedChild = Child->CastTo<InItemType>())
-			{
-				OutDescendants.Add(CastedChild);
-			}
-
-			// Recursively check the child's children
-			if (Child.IsValid())
-			{
-				OutDescendants.Append(Child->GetDescendantsOfType<InItemType>());
-			}
-		}
-
-		return OutDescendants;
-	}
-
-	template<typename InItemType>
-	TArray<InItemType*> GetChildrenOfType(const bool bInIncludeThis = false) const
-	{
-		TArray<InItemType*> OutChildren;
-
-		if (bInIncludeThis)
-		{
-			if (InItemType* const CastedChild = const_cast<InItemType*>(CastTo<const InItemType>()))
-			{
-				OutChildren.Add(CastedChild);
-			}
-		}
-
-		for (const FNavigationToolItemPtr& Child : GetChildren())
-		{
-			if (InItemType* const CastedChild = Child->CastTo<InItemType>())
-			{
-				OutChildren.Add(CastedChild);
-			}
-		}
-
-		return OutChildren;
-	};
+	virtual FNavigationToolViewModelWeakPtr GetChildAt(const IndexType InIndex) const = 0;
 
 	/**
-	 * Finds the first ancestor view model that implements the given type
+	 * Gets the Parent of this item.
+	 * Should only be null prior to registering it in Navigation Tool or if its root item.
 	 */
-	template<typename InItemType>
-	InItemType* FindAncestorOfType(const bool bInIncludeThis = false)
-	{
-		if (bInIncludeThis && IsA<InItemType>())
-		{
-			return CastTo<InItemType>();
-		}
-
-		FNavigationToolItemPtr CurrentParent = GetParent();
-		while (CurrentParent.IsValid())
-		{
-			if (CurrentParent->IsA<InItemType>())
-			{
-				return CurrentParent->CastTo<InItemType>();
-			}
-
-			CurrentParent = CurrentParent->GetParent();
-		}
-
-		return nullptr;
-	}
+	virtual FNavigationToolViewModelPtr GetParent() const = 0;
+
+	/** Gets the list of all parents up the tree from this Item */
+	virtual TSet<FNavigationToolViewModelPtr> GetParents(const bool bInIncludeRoot = false) const = 0;
+
+	/** Sets the parent. Note that the parent must've already had this instance as a child (check is done) */
+	virtual void SetParent(FNavigationToolViewModelPtr InParent) = 0;
+
+	/** Gets the Id of this item */
+	virtual FNavigationToolItemId GetItemId() const = 0;
 
 	/**
-	 * Whether this Item can be at the Top Level just beneath the Root, or it needs to always be under some other Item
+	 * Whether this item can be at the Top Level just beneath the Root, or it needs to always be under some other item
 	 * E.g. Actors can be Top Level, but Components or Materials can't
 	 */
 	virtual bool CanBeTopLevel() const = 0;
-
-	/** Gets the Id of this Item */
-	virtual FNavigationToolItemId GetItemId() const = 0;
 	
-	/** Returns whether this Item (and what it represents) should be allowed to be registered in Navigation Tool */
+	/** Returns whether this item (and what it represents) should be allowed to be registered in Navigation Tool */
 	virtual bool IsAllowedInTool() const = 0;
 	
-	/** Gets the Display Name Text of the Item */
+	/** Gets the display name text of the item */
 	virtual FText GetDisplayName() const = 0;
 
 	/** Gets the Class/Type of this Item (e.g. for Items that represent UObjects, it will be the UObject class) */
@@ -444,29 +207,14 @@ public:
 	virtual ENavigationToolItemViewMode GetSupportedViewModes(const INavigationToolView& InToolView) const = 0;
 
 	/** Whether this Item should be visualized in the given View Mode, for the given Navigation Tool View */
-	virtual bool IsViewModeSupported(const ENavigationToolItemViewMode InViewMode, const INavigationToolView& InToolView) const
-	{
-		return EnumHasAnyFlags(InViewMode, GetSupportedViewModes(InToolView));
-	};
+	virtual bool IsViewModeSupported(const ENavigationToolItemViewMode InViewMode, const INavigationToolView& InToolView) const = 0;
 
 	/**
 	 * Called when objects have been replaced on the Engine side. Used to replace any UObjects used by this item
 	 * @param InReplacementMap the map of old object that is garbage to the new object that replaced it
 	 * @param bInRecursive     whether to recurse this same function to children items
 	 */
-	virtual void OnObjectsReplaced(const TMap<UObject*, UObject*>& InReplacementMap, const bool bInRecursive)
-	{
-		if (bInRecursive)
-		{
-			for (const FNavigationToolItemPtr& ChildItem : GetChildren())
-			{
-				if (ChildItem.IsValid())
-				{
-					ChildItem->OnObjectsReplaced(InReplacementMap, bInRecursive);
-				}
-			}
-		}
-	}
+	virtual void OnObjectsReplaced(const TMap<UObject*, UObject*>& InReplacementMap, const bool bInRecursive) {}
 
 	/** Function responsible of Generating the Label Widget for this Item (i.e. the column containing the Icon and the Name) */
 	virtual TSharedRef<SWidget> GenerateLabelWidget(const TSharedRef<SNavigationToolTreeRow>& InRow) = 0;
@@ -488,10 +236,10 @@ public:
 	 * const TSharedPtr<FNavigationToolView>& - The Navigation Tool View where this Expansion happened
 	 * bool - Whether the Item is expanded
 	 */
-	using FOnExpansionChanged = TMulticastDelegate<void(const TSharedPtr<INavigationToolView>&, bool)>;
+	//using FOnExpansionChanged = TMulticastDelegate<void(const TSharedPtr<INavigationToolView>&, bool)>;
 
 	/** Called when Expansion state (Expanded/Collapsed) has been changed */
-	virtual FOnExpansionChanged& OnExpansionChanged() = 0;
+	//virtual FOnExpansionChanged& OnExpansionChanged() = 0;
 
 	/** Whether the Item is able to expand when AutoExpand functionality is enabled */
 	virtual bool CanAutoExpand() const = 0;
@@ -511,40 +259,30 @@ public:
 	virtual void SetFlags(const ENavigationToolItemFlags InFlags) = 0;
 	virtual ENavigationToolItemFlags GetFlags() const = 0;
 
-	/** Gets the Tags found for this Item (e.g. for Actors, actor tags and for Components Component Tags) */
+	/** Gets the Tags found for this item (e.g. for Actors, actor tags and for Components, component tags) */
 	virtual TArray<FName> GetTags() const { return TArray<FName>(); }
 
-	/** Returns the item's height in tree, Root Item should return 0 as it has no Parent */
-	virtual int32 GetItemTreeHeight() const
-	{
-		int32 Height = 0;
-		FNavigationToolItemPtr TopParent = GetParent();
-		while (TopParent.IsValid())
-		{
-			TopParent = TopParent->GetParent();
-			++Height;
-		}
-		return Height;
-	}
+	/** Returns the item's height in the tree. Root item should return 0 as it has no parent. */
+	//virtual int32 GetItemTreeHeight() const = 0;
 
 	/**
 	 * Delegate signature for relaying an Item Rename action
 	 * ENavigationToolRenameAction the type of action being relayed (e.g. request a rename, or notify rename complete, etc)
 	 * const TSharedPtr<INavigationToolView>& - The Navigation Tool View where the rename action is taking place
 	 */
-	using FOnRenameAction = TMulticastDelegate<void(ENavigationToolRenameAction, const TSharedPtr<INavigationToolView>&)>;
+	//using FOnRenameAction = TMulticastDelegate<void(ENavigationToolRenameAction, const TSharedPtr<INavigationToolView>&)>;
 
 	/** Broadcasts whenever a rename action takes place from a given view (e.g. when pressing "F2" to rename, or committing the rename text) */
-	virtual FOnRenameAction& OnRenameAction() = 0;
+	//virtual FOnRenameAction& OnRenameAction() = 0;
 
 	/** Determines if and where the incoming Drag Drop Event can be processed by this item */
 	virtual TOptional<EItemDropZone> CanAcceptDrop(const FDragDropEvent& InDragDropEvent, const EItemDropZone InDropZone) = 0;
 
 	/** Processes the Drag and Drop Event for this Item */
-	virtual FReply AcceptDrop(const FDragDropEvent& InDragDropEvent, const EItemDropZone InDropZone) = 0;
+	virtual  FReply AcceptDrop(const FDragDropEvent& InDragDropEvent, const EItemDropZone InDropZone) = 0;
 
 	/** Whether Ignoring Pending Kill. Useful to get underlying UObjects that are pending kill and get the pointer to it and not a null value */
-	bool IsIgnoringPendingKill() const { return HasAllFlags(ENavigationToolItemFlags::IgnorePendingKill); }
+	virtual bool IsIgnoringPendingKill() const = 0;
 
 	/** Gets whether this item is expanded */
 	virtual bool IsExpanded() const = 0;
@@ -553,37 +291,20 @@ public:
 	virtual void SetExpansion(const bool bInIsExpanded) = 0;
 
 	/** Converts this Navigation Tool item to a serialized item that can be saved in the sequence */
-	virtual FNavigationToolSerializedItem MakeSerializedItem()
-	{
-		return FNavigationToolSerializedItem(GetItemId().GetStringId());
-	}
-};
+	virtual FNavigationToolSerializedItem MakeSerializedItem() = 0;
 
-/** Adds Scoped Item Flags, removes them when out of scope. Useful for temp checks like IgnorePendingKill */
-struct FNavigationToolItemFlagGuard
-{
-	FNavigationToolItemFlagGuard(const FNavigationToolItemPtr InItem, const ENavigationToolItemFlags ItemFlags)
-	{
-		if (InItem.IsValid())
-		{
-			Item         = InItem;
-			OldItemFlags = InItem->GetFlags();
-			InItem->SetFlags(ItemFlags);
-		}
-	}
-
-	~FNavigationToolItemFlagGuard()
-	{
-		if (Item.IsValid())
-		{
-			Item->SetFlags(OldItemFlags);
-		}
-	}
+	FOnRenameAction& OnRenameAction() { return OnRenameActionDelegate; }
+
+	FOnExpansionChanged& OnExpansionChanged() { return OnExpansionChangedDelegate; }
 
 protected:
-	FNavigationToolItemPtr Item;
+	/** Delegate for when expansion changes in the item */
+	FOnExpansionChanged OnExpansionChangedDelegate;
 	
-	ENavigationToolItemFlags OldItemFlags;
+	/** The delegate for renaming */
+	FOnRenameAction OnRenameActionDelegate;
 };
 
 } // namespace UE::SequenceNavigator
+
+#undef UE_API
diff --git a/Engine/Plugins/Experimental/SequenceNavigator/Source/SequenceNavigator/Public/Items/NavigationToolActor.h b/Engine/Plugins/Experimental/SequenceNavigator/Source/SequenceNavigator/Public/Items/NavigationToolActor.h
--- a/Engine/Plugins/Experimental/SequenceNavigator/Source/SequenceNavigator/Public/Items/NavigationToolActor.h
+++ b/Engine/Plugins/Experimental/SequenceNavigator/Source/SequenceNavigator/Public/Items/NavigationToolActor.h
@@ -2,6 +2,7 @@
 
 #pragma once
 
+#include "MVVM/ViewModelTypeID.h"
 #include "NavigationToolBinding.h"
 #include "NavigationToolDefines.h"
 
@@ -21,17 +22,17 @@ class FNavigationToolActor
 	: public FNavigationToolBinding
 {
 public:
-	UE_NAVIGATIONTOOL_INHERITS_WITH_SUPER(FNavigationToolActor
-		, FNavigationToolBinding);
+	UE_SEQUENCER_DECLARE_CASTABLE(FNavigationToolActor
+		, FNavigationToolBinding)
 
 	UE_API FNavigationToolActor(INavigationTool& InTool
-		, const FNavigationToolItemPtr& InParentItem
+		, const FNavigationToolViewModelPtr& InParentItem
 		, const TSharedPtr<FNavigationToolSequence>& InParentSequenceItem
 		, const FMovieSceneBinding& InBinding);
 
 	//~ Begin INavigationToolItem
 
-	UE_API virtual void FindChildren(TArray<FNavigationToolItemPtr>& OutChildren, const bool bInRecursive) override;
+	UE_API virtual void FindChildren(TArray<FNavigationToolViewModelWeakPtr>& OutWeakChildren, const bool bInRecursive) override;
 
 	UE_API virtual bool IsAllowedInTool() const override;
 	UE_API virtual ENavigationToolItemViewMode GetSupportedViewModes(const INavigationToolView& InToolView) const override;
@@ -47,7 +48,7 @@ public:
 
 	//~ Begin IRenameableExtension
 	UE_API virtual bool CanRename() const override;
-	UE_API virtual bool Rename(const FString& InName) override;
+	UE_API virtual void Rename(const FText& InNewName) override;
 	//~ End IRenameableExtension
 
 	UE_API AActor* GetActor() const;
diff --git a/Engine/Plugins/Experimental/SequenceNavigator/Source/SequenceNavigator/Public/Items/NavigationToolBinding.h b/Engine/Plugins/Experimental/SequenceNavigator/Source/SequenceNavigator/Public/Items/NavigationToolBinding.h
--- a/Engine/Plugins/Experimental/SequenceNavigator/Source/SequenceNavigator/Public/Items/NavigationToolBinding.h
+++ b/Engine/Plugins/Experimental/SequenceNavigator/Source/SequenceNavigator/Public/Items/NavigationToolBinding.h
@@ -3,12 +3,11 @@
 #pragma once
 
 #include "Extensions/IPlayheadExtension.h"
-#include "Extensions/IRenameableExtension.h"
-#include "Extensions/ISequenceLockableExtension.h"
 #include "MovieSceneBinding.h"
+#include "MVVM/Extensions/IRenameableExtension.h"
+#include "MVVM/ViewModelTypeID.h"
 #include "MVVM/ViewModels/ObjectBindingModel.h"
 #include "NavigationToolItem.h"
-#include "SequencerCoreFwd.h"
 #include "Styling/SlateColor.h"
 #include "Textures/SlateIcon.h"
 
@@ -29,19 +28,19 @@ class FNavigationToolSequence;
  */
 class FNavigationToolBinding
 	: public FNavigationToolItem
-	, public IRenameableExtension
+	, public Sequencer::IRenameableExtension
 	//, public ISequenceLockableExtension
 	, public IPlayheadExtension
 {
 public:
-	UE_NAVIGATIONTOOL_INHERITS_WITH_SUPER(FNavigationToolBinding
+	UE_SEQUENCER_DECLARE_CASTABLE_API(UE_API, FNavigationToolBinding
 		, FNavigationToolItem
-		, IRenameableExtension
-		/*, ISequenceLockableExtension*/
+		, Sequencer::IRenameableExtension
+		//, ISequenceLockableExtension
 		, IPlayheadExtension);
 
 	UE_API FNavigationToolBinding(INavigationTool& InTool
-		, const FNavigationToolItemPtr& InParentItem
+		, const FNavigationToolViewModelPtr& InParentItem
 		, const TSharedPtr<FNavigationToolSequence>& InParentSequenceItem
 		, const FMovieSceneBinding& InBinding);
 
@@ -51,7 +50,7 @@ public:
 	UE_API virtual UObject* GetItemObject() const override;
 	UE_API virtual bool IsAllowedInTool() const override;
 
-	UE_API virtual void FindChildren(TArray<FNavigationToolItemPtr>& OutChildren, const bool bInRecursive) override;
+	UE_API virtual void FindChildren(TArray<FNavigationToolViewModelWeakPtr>& OutWeakChildren, const bool bInRecursive) override;
 
 	virtual bool CanBeTopLevel() const override { return false; }
 	virtual bool ShouldSort() const override { return false; }
@@ -76,7 +75,7 @@ public:
 
 	//~ Begin IRenameableExtension
 	UE_API virtual bool CanRename() const override;
-	UE_API virtual bool Rename(const FString& InName) override;
+	UE_API virtual void Rename(const FText& InNewName) override;
 	//~ End IRenameableExtension
 
 	//~ Begin ISequenceLockableExtension
@@ -105,7 +104,7 @@ public:
 	UE_API UE::Sequencer::TViewModelPtr<UE::Sequencer::FObjectBindingModel> GetViewModel() const;
 
 protected:
-	TWeakPtr<FNavigationToolSequence> WeakParentSequenceItem;
+	Sequencer::TWeakViewModelPtr<FNavigationToolSequence> WeakParentSequenceItem;
 
 	FMovieSceneBinding Binding;
 
diff --git a/Engine/Plugins/Experimental/SequenceNavigator/Source/SequenceNavigator/Public/Items/NavigationToolComponent.h b/Engine/Plugins/Experimental/SequenceNavigator/Source/SequenceNavigator/Public/Items/NavigationToolComponent.h
--- a/Engine/Plugins/Experimental/SequenceNavigator/Source/SequenceNavigator/Public/Items/NavigationToolComponent.h
+++ b/Engine/Plugins/Experimental/SequenceNavigator/Source/SequenceNavigator/Public/Items/NavigationToolComponent.h
@@ -2,6 +2,7 @@
 
 #pragma once
 
+#include "MVVM/ViewModelTypeID.h"
 #include "NavigationToolBinding.h"
 
 #define UE_API SEQUENCENAVIGATOR_API
@@ -18,17 +19,17 @@ class FNavigationToolComponent
 	: public FNavigationToolBinding
 {
 public:
-	UE_NAVIGATIONTOOL_INHERITS_WITH_SUPER(FNavigationToolComponent
-		, FNavigationToolBinding);
+	UE_SEQUENCER_DECLARE_CASTABLE_API(UE_API, FNavigationToolComponent
+		, FNavigationToolBinding)
 
 	UE_API FNavigationToolComponent(INavigationTool& InTool
-		, const FNavigationToolItemPtr& InParentItem
+		, const FNavigationToolViewModelPtr& InParentItem
 		, const TSharedPtr<FNavigationToolSequence>& InParentSequenceItem
 		, const FMovieSceneBinding& InBinding);
 
 	//~ Begin INavigationToolItem
 
-	UE_API virtual void FindChildren(TArray<FNavigationToolItemPtr>& OutChildren, const bool bInRecursive) override;
+	UE_API virtual void FindChildren(TArray<FNavigationToolViewModelWeakPtr>& OutWeakChildren, const bool bInRecursive) override;
 	UE_API virtual void GetItemProxies(TArray<TSharedPtr<FNavigationToolItemProxy>>& OutItemProxies) override;
 
 	UE_API virtual bool IsAllowedInTool() const override;
@@ -47,7 +48,7 @@ public:
 
 	//~ Begin IRenameableExtension
 	UE_API virtual bool CanRename() const override;
-	UE_API virtual bool Rename(const FString& InName) override;
+	UE_API virtual void Rename(const FText& InNewName) override;
 	//~ End IRenameableExtension
 
 	UE_API UActorComponent* GetComponent() const;
diff --git a/Engine/Plugins/Experimental/SequenceNavigator/Source/SequenceNavigator/Public/Items/NavigationToolItem.h b/Engine/Plugins/Experimental/SequenceNavigator/Source/SequenceNavigator/Public/Items/NavigationToolItem.h
--- a/Engine/Plugins/Experimental/SequenceNavigator/Source/SequenceNavigator/Public/Items/NavigationToolItem.h
+++ b/Engine/Plugins/Experimental/SequenceNavigator/Source/SequenceNavigator/Public/Items/NavigationToolItem.h
@@ -2,86 +2,132 @@
 
 #pragma once
 
-#include "Misc/Optional.h"
+#include "EditorSequenceNavigationDefs.h"
 #include "Extensions/IColorExtension.h"
 #include "Items/INavigationToolItem.h"
 #include "Items/NavigationToolItemId.h"
-#include "NavigationToolDefines.h"
-#include "NavigationToolItemType.h"
+#include "Misc/Optional.h"
+#include "MVVM/Extensions/IOutlinerExtension.h"
+#include "MVVM/ViewModels/OutlinerItemModel.h"
 #include "NavigationToolSettings.h"
+#include "ToolkitBuilder.h"
 
 #define UE_API SEQUENCENAVIGATOR_API
 
+class FText;
+
 namespace UE::SequenceNavigator
 {
 
+class FNavigationToolProvider;
+class FNavigationToolTreeRoot;
+class INavigationTool;
+class INavigationToolView;
 class SNavigationToolTreeRow;
+struct FNavigationToolAddItemParams;
+struct FNavigationToolRemoveItemParams;
 
-/** Base Implementation of INavigationToolItem */
+/** Base implementation of INavigationToolItem */
 class FNavigationToolItem
-	: public INavigationToolItem
+	: public Sequencer::TOutlinerModelMixin<Sequencer::FViewModel>
+	, public INavigationToolItem
 	, public IColorExtension
 {
 	friend class FNavigationTool;
 
 public:
-	UE_NAVIGATIONTOOL_INHERITS_WITH_SUPER(FNavigationToolItem
+	using IndexType = TArray<INavigationToolItem>::SizeType;
+
+	UE_SEQUENCER_DECLARE_CASTABLE_API(UE_API, FNavigationToolItem
+		, Sequencer::TOutlinerModelMixin<Sequencer::FViewModel>
 		, INavigationToolItem
-		, IColorExtension);
+		, IColorExtension)
 
-	DECLARE_DELEGATE_RetVal_OneParam(FNavigationToolSaveState*, FNavigationToolGetToolState, const FNavigationToolItemPtr /*InItem*/);
+	DECLARE_DELEGATE_RetVal_OneParam(FNavigationToolSaveState*, FNavigationToolGetToolState
+		, const FNavigationToolViewModelPtr /*InItem*/);
 
-	UE_API FNavigationToolItem(INavigationTool& InTool, const FNavigationToolItemPtr& InParentItem);
+	UE_API FNavigationToolItem(INavigationTool& InTool, const FNavigationToolViewModelPtr& InParentItem);
 
 	//~ Begin INavigationToolItem
+
 	UE_API virtual INavigationTool& GetOwnerTool() const final override;
 	UE_API virtual TSharedPtr<FNavigationToolProvider> GetProvider() const final override;
 	UE_API virtual FNavigationToolSaveState* GetProviderSaveState() const override;
 	UE_API virtual bool IsItemValid() const override;
+
 	UE_API virtual void RefreshChildren() override;
 	UE_API virtual void ResetChildren() override;
-	virtual const TArray<FNavigationToolItemPtr>& GetChildren() const override { return Children; }
-	UE_API virtual void FindChildren(TArray<FNavigationToolItemPtr>& OutChildren, const bool bInRecursive) override;
-	UE_API virtual void FindValidChildren(TArray<FNavigationToolItemPtr>& OutChildren, const bool bInRecursive) override;
-	UE_API virtual TArray<FNavigationToolItemPtr> FindPath(const TArray<FNavigationToolItemPtr>& InItems) const override;
-	virtual TArray<FNavigationToolItemPtr>& GetChildrenMutable() override { return Children; }
-	UE_API virtual IndexType GetChildIndex(const FNavigationToolItemPtr& ChildItem) const override;
+
+	virtual const TArray<FNavigationToolViewModelWeakPtr>& GetChildren() const override { return WeakChildren; }
+	virtual TArray<FNavigationToolViewModelWeakPtr>& GetChildrenMutable() override { return WeakChildren; }
+
+	UE_API virtual void FindChildren(TArray<FNavigationToolViewModelWeakPtr>& OutWeakChildren
+		, const bool bInRecursive) override;
+	UE_API virtual void FindValidChildren(TArray<FNavigationToolViewModelWeakPtr>& OutWeakChildren
+		, const bool bInRecursive) override;
+
+	UE_API virtual TArray<FNavigationToolViewModelPtr> FindPath(const TArray<FNavigationToolViewModelWeakPtr>& InWeakItems) const override final;
+	UE_API virtual FString GetFullPath() const override final;
+
+	UE_API virtual IndexType GetChildIndex(const FNavigationToolViewModelWeakPtr& InWeakChildItem) const override;
+	UE_API virtual FNavigationToolViewModelWeakPtr GetChildAt(const IndexType InIndex) const override;
+
+	UE_API virtual FNavigationToolViewModelPtr GetParent() const override;
+	UE_API virtual TSet<FNavigationToolViewModelPtr> GetParents(const bool bInIncludeRoot = false) const override;
+	UE_API virtual void SetParent(FNavigationToolViewModelPtr InParent) override;
+
 	virtual bool ShouldSort() const override { return false; }
-	UE_API virtual bool CanAddChild(const FNavigationToolItemPtr& InChild) const override;
+
+	UE_API virtual bool CanAddChild(const FNavigationToolViewModelPtr& InChild) const override;
 	UE_API virtual bool AddChild(const FNavigationToolAddItemParams& InAddItemParams) override;
 	UE_API virtual bool RemoveChild(const FNavigationToolRemoveItemParams& InRemoveItemParams) override;
-	virtual FNavigationToolItemPtr GetParent() const override { return WeakParent.Pin(); }
-	UE_API virtual void SetParent(FNavigationToolItemPtr InParent) override;
+
+	UE_API virtual FNavigationToolItemId GetItemId() const override final;
 	virtual bool CanBeTopLevel() const override { return false; }
 	virtual bool IsAllowedInTool() const override { return true; }
-	virtual FText GetDisplayName() const override { return FText::GetEmpty(); }
-	virtual FText GetClassName() const override { return FText::GetEmpty(); }
+	UE_API virtual FText GetDisplayName() const override { return FText::GetEmpty(); }
+	UE_API virtual FText GetClassName() const override { return FText::GetEmpty(); }
+
 	UE_API virtual ENavigationToolItemViewMode GetSupportedViewModes(const INavigationToolView& InToolView) const override;
-	UE_API virtual FNavigationToolItemId GetItemId() const override final;
+	UE_API virtual bool IsViewModeSupported(const ENavigationToolItemViewMode InViewMode, const INavigationToolView& InToolView) const override;
+
 	UE_API virtual FSlateColor GetItemLabelColor() const override;
 	UE_API virtual FLinearColor GetItemTintColor() const override;
 	UE_API virtual const FSlateBrush* GetIconBrush() const override;
 	UE_API virtual FSlateColor GetIconColor() const override;
+
 	UE_API virtual TSharedRef<SWidget> GenerateLabelWidget(const TSharedRef<SNavigationToolTreeRow>& InRow) override;
+
 	virtual bool ShowVisibility() const override { return false; }
 	virtual bool CanReceiveParentVisibilityPropagation() const override { return false; }
 	virtual bool GetVisibility() const override { return false; }
 	virtual bool CanAutoExpand() const override { return true; }
 	virtual bool CanDelete() const override { return false; }
 	UE_API virtual bool Delete() override;
+
 	UE_API virtual void AddFlags(ENavigationToolItemFlags Flags) override;
 	UE_API virtual void RemoveFlags(ENavigationToolItemFlags Flags) override;
 	UE_API virtual bool HasAnyFlags(ENavigationToolItemFlags Flags) const override;
 	UE_API virtual bool HasAllFlags(ENavigationToolItemFlags Flags) const override;
 	virtual void SetFlags(ENavigationToolItemFlags InFlags) override { ItemFlags = InFlags; }
 	virtual ENavigationToolItemFlags GetFlags() const override { return ItemFlags; }
-	virtual TArray<FName> GetTags() const override { return TArray<FName>(); };
-	UE_API virtual TOptional<EItemDropZone> CanAcceptDrop(const FDragDropEvent& InDragDropEvent, EItemDropZone InDropZone) override;
-	UE_API virtual FReply AcceptDrop(const FDragDropEvent& InDragDropEvent, EItemDropZone InDropZone) override;
-	virtual FOnRenameAction& OnRenameAction() override { return OnRenameActionDelegate; }
-	virtual FOnExpansionChanged& OnExpansionChanged() override { return OnExpansionChangedDelegate; }
+
+	virtual TArray<FName> GetTags() const override { return TArray<FName>(); }
+
+	//virtual int32 GetItemTreeHeight() const override;
+
+	UE_API virtual TOptional<EItemDropZone> CanAcceptDrop(const FDragDropEvent& InDragDropEvent, const EItemDropZone InDropZone) override;
+	UE_API virtual FReply AcceptDrop(const FDragDropEvent& InDragDropEvent, const EItemDropZone InDropZone) override;
+
+	UE_API virtual bool IsIgnoringPendingKill() const override final;
+
+	UE_API virtual void OnObjectsReplaced(const TMap<UObject*, UObject*>& InReplacementMap, const bool bInRecursive) override;
+
 	UE_API virtual bool IsExpanded() const override;
 	UE_API virtual void SetExpansion(const bool bInExpand) override;
+
+	UE_API virtual FNavigationToolSerializedItem MakeSerializedItem() override final;
+
 	//~ End INavigationToolItem
 
 	//~ Begin IColorExtension
@@ -89,6 +135,32 @@ public:
 	UE_API virtual void SetColor(const TOptional<FColor>& InColor) override;
 	//~ End IColorExtension
 
+	//~ Begin IOutlinerExtension
+	UE_API virtual Sequencer::FOutlinerSizing GetOutlinerSizing() const override;
+	//~ End IOutlinerExtension
+
+	Sequencer::TViewModelPtr<FNavigationToolItem> AsItemViewModel()
+	{
+		const TSharedRef<FNavigationToolItem> SharedThisRef = SharedThis(this);
+		return Sequencer::CastViewModel<FNavigationToolItem>(SharedThisRef);
+		//return TOutlinerModelMixin::AsViewModel()->CastThisShared<FNavigationToolItem>();
+	}
+
+	Sequencer::TViewModelPtr<FNavigationToolItem> AsItemViewModelConst() const
+	{
+		//if (const FViewModel* const ViewModel = AsViewModel())
+		{
+			const TSharedRef<FNavigationToolItem> SharedThisRef = SharedThis(const_cast<FNavigationToolItem*>(this));
+			return Sequencer::CastViewModel<FNavigationToolItem>(SharedThisRef);
+		}
+		//return nullptr;
+/*
+		//TSharedRef<FNavigationToolItem> SharedThisRef = SharedThis(const_cast<FNavigationToolItem*>(this));
+		auto SharedThisRef = AsShared();
+		const FNavigationToolItem* const Item = SharedThisRef->CastThis<FNavigationToolItem>();
+		return Sequencer::TViewModelPtr(SharedThisRef, Item);// ? Item : CastDynamic<FNavigationToolItem>());*/
+	}
+
 protected:
 	friend class INavigationTool;
 
@@ -100,14 +172,14 @@ protected:
 	/** Sets the ItemId member var to what CalculateItemId returns */
 	UE_API void RecalculateItemId();
 	
-	/** The actual implementation of putting the given Item under the Children array */
+	/** The actual implementation of putting the given item under the children array */
 	UE_API void AddChildChecked(const FNavigationToolAddItemParams& InAddItemParams);
 
 	/** The actual implementation of removing the given item from the Children array */
 	UE_API bool RemoveChildChecked(const FNavigationToolRemoveItemParams& InRemoveItemParams);
 	
 	/** Careful handling of multiple children being detected and added to this item children array */
-	UE_API void HandleNewSortableChildren(TArray<FNavigationToolItemPtr> InSortableChildren);
+	UE_API void HandleNewSortableChildren(TArray<FNavigationToolViewModelWeakPtr> InWeakSortableChildren);
 
 	/** Reference to the Owning Navigation Tool */
 	INavigationTool& Tool;
@@ -116,23 +188,44 @@ protected:
 	TWeakPtr<FNavigationToolProvider> WeakProvider;
 
 	/** Weak pointer to the Parent Item. Can be null, but if valid, the Parent should have this item in the Children Array */
-	FNavigationToolItemWeakPtr WeakParent;
+	FNavigationToolViewModelWeakPtr WeakParent;
 
 	/** Array of Shared pointers to the Child Items. These Items should have their ParentWeak pointing to this item */
-	TArray<FNavigationToolItemPtr> Children;
-
-	/** Delegate for when Expansion changes in the Item */
-	FOnExpansionChanged OnExpansionChangedDelegate;
-	
-	/** The delegate for renaming */
-	FOnRenameAction OnRenameActionDelegate;
+	TArray<FNavigationToolViewModelWeakPtr> WeakChildren;
 
 	/** The current flags set for this item */
-	ENavigationToolItemFlags ItemFlags = ENavigationToolItemFlags::None;
+	ENavigationToolItemFlags ItemFlags;
 
 	FNavigationToolItemId ItemId;
 };
 
+/** Adds scoped item flags and removes them when out of scope. Useful for temp checks like IgnorePendingKill */
+struct FNavigationToolItemFlagGuard
+{
+	FNavigationToolItemFlagGuard(const FNavigationToolViewModelPtr& InItem, const ENavigationToolItemFlags InItemFlags)
+	{
+		if (InItem.IsValid())
+		{
+			WeakItem = InItem;
+			OldItemFlags = InItem->GetFlags();
+			InItem->SetFlags(InItemFlags);
+		}
+	}
+
+	~FNavigationToolItemFlagGuard()
+	{
+		if (const FNavigationToolViewModelPtr Item = WeakItem.Pin())
+		{
+			Item->SetFlags(OldItemFlags);
+		}
+	}
+
+protected:
+	FNavigationToolViewModelWeakPtr WeakItem;
+	
+	ENavigationToolItemFlags OldItemFlags;
+};
+
 } // namespace UE::SequenceNavigator
 
 #undef UE_API
diff --git a/Engine/Plugins/Experimental/SequenceNavigator/Source/SequenceNavigator/Public/Items/NavigationToolItemId.h b/Engine/Plugins/Experimental/SequenceNavigator/Source/SequenceNavigator/Public/Items/NavigationToolItemId.h
--- a/Engine/Plugins/Experimental/SequenceNavigator/Source/SequenceNavigator/Public/Items/NavigationToolItemId.h
+++ b/Engine/Plugins/Experimental/SequenceNavigator/Source/SequenceNavigator/Public/Items/NavigationToolItemId.h
@@ -2,7 +2,7 @@
 
 #pragma once
 
-#include "ItemProxies/INavigationToolItemProxyFactory.h"
+//#include "ItemProxies/INavigationToolItemProxyFactory.h"
 #include "NavigationToolDefines.h"
 #include "SequencerCoreFwd.h"
 
@@ -16,6 +16,7 @@ namespace UE::SequenceNavigator
 {
 
 class INavigationToolItem;
+class INavigationToolItemProxyFactory;
 
 /** Struct to identify an item in the Navigation Tool */
 struct FNavigationToolItemId
@@ -50,7 +51,7 @@ struct FNavigationToolItemId
 	 * @param InReferencingId The Id to differentiate this reference from other references within the same object (e.g. can be slot or property name)
 	 */
 	UE_API FNavigationToolItemId(const UObject* InObject
-		, const FNavigationToolItemPtr& InReferencingItem
+		, const FNavigationToolViewModelPtr& InReferencingItem
 		, const FString& InReferencingId);
 
 	/**
@@ -59,7 +60,7 @@ struct FNavigationToolItemId
 	 * @param InParentItem The parent item holding the item proxy
 	 * @param InItemProxyFactory The factory responsible for creating the item proxy
 	 */
-	UE_API FNavigationToolItemId(const FNavigationToolItemPtr& InParentItem
+	UE_API FNavigationToolItemId(const FNavigationToolViewModelPtr& InParentItem
 		, const INavigationToolItemProxyFactory& InItemProxyFactory);
 
 	/**
@@ -67,8 +68,8 @@ struct FNavigationToolItemId
 	 * @param InParentItem The parent item holding the item proxy
 	 * @param InItemProxy The item proxy under the parent item
 	 */
-	UE_API FNavigationToolItemId(const FNavigationToolItemPtr& InParentItem
-		, const INavigationToolItem& InItemProxy);
+	UE_API FNavigationToolItemId(const FNavigationToolViewModelPtr& InParentItem
+		, const FNavigationToolViewModelPtr& InItemProxy);
 
 	/**
 	 * Ctor used for making the Item Id for most sequence items
@@ -78,7 +79,7 @@ struct FNavigationToolItemId
 	 * @param InSectionIndex The parent section index, if any. Should be specified if InSection is used.
 	 * @param InReferenceId The optional additional reference string
 	 */
-	UE_API FNavigationToolItemId(const FNavigationToolItemPtr& InParentItem
+	UE_API FNavigationToolItemId(const FNavigationToolViewModelPtr& InParentItem
 		, const UMovieSceneSequence* const InSequence
 		, const UMovieSceneSection* const InSection = nullptr
 		, const int32 InSectionIndex = 0
@@ -117,7 +118,7 @@ private:
 	 * @param InSectionIndex The parent section index, if any. Should be specified if InSection is used.
 	 * @param InReferenceId The optional additional reference string
 	 */
-	UE_API void ConstructId(const FNavigationToolItemPtr& InParentItem
+	UE_API void ConstructId(const FNavigationToolViewModelPtr& InParentItem
 		, const UMovieSceneSequence* const InSequence
 		, const UMovieSceneSection* const InSection = nullptr
 		, const int32 InSectionIndex = 0
diff --git a/Engine/Plugins/Experimental/SequenceNavigator/Source/SequenceNavigator/Public/Items/NavigationToolItemParameters.h b/Engine/Plugins/Experimental/SequenceNavigator/Source/SequenceNavigator/Public/Items/NavigationToolItemParameters.h
--- a/Engine/Plugins/Experimental/SequenceNavigator/Source/SequenceNavigator/Public/Items/NavigationToolItemParameters.h
+++ b/Engine/Plugins/Experimental/SequenceNavigator/Source/SequenceNavigator/Public/Items/NavigationToolItemParameters.h
@@ -26,22 +26,22 @@ ENUM_CLASS_FLAGS(ENavigationToolAddItemFlags);
 
 struct FNavigationToolAddItemParams
 {
-	FNavigationToolAddItemParams(const FNavigationToolItemPtr& InItem = nullptr
+	FNavigationToolAddItemParams(const FNavigationToolViewModelPtr& InItem = nullptr
 			, const ENavigationToolAddItemFlags InFlags = ENavigationToolAddItemFlags::None
-			, const FNavigationToolItemPtr& InRelativeItem = nullptr
+			, const FNavigationToolViewModelPtr& InRelativeItem = nullptr
 			, const TOptional<EItemDropZone>& InRelativeDropZone = TOptional<EItemDropZone>())
-		: Item(InItem)
-		, RelativeItem(InRelativeItem)
+		: WeakItem(InItem)
+		, WeakRelativeItem(InRelativeItem)
 		, RelativeDropZone(InRelativeDropZone)
 		, Flags(InFlags)
 		, SelectionFlags(ENavigationToolItemSelectionFlags::None)
 	{}
 
 	/** The Item to Add */
-	FNavigationToolItemPtr Item;
+	FNavigationToolViewModelWeakPtr WeakItem;
 
 	/** The Item to use as base in where to place the Item */
-	FNavigationToolItemPtr RelativeItem;
+	FNavigationToolViewModelWeakPtr WeakRelativeItem;
 
 	/** The Placement Order from the Relative Item (Onto/Inside, Above, Below) */
 	TOptional<EItemDropZone> RelativeDropZone;
@@ -55,11 +55,11 @@ struct FNavigationToolAddItemParams
 
 struct FNavigationToolRemoveItemParams
 {
-	FNavigationToolRemoveItemParams(const FNavigationToolItemPtr& InItem = nullptr)
-		: Item(InItem)
+	FNavigationToolRemoveItemParams(const FNavigationToolViewModelPtr& InItem = nullptr)
+		: WeakItem(InItem)
 	{}
 
-	FNavigationToolItemPtr Item;
+	FNavigationToolViewModelWeakPtr WeakItem;
 };
 
-}
+} // namespace UE::SequenceNavigator
diff --git a/Engine/Plugins/Experimental/SequenceNavigator/Source/SequenceNavigator/Public/Items/NavigationToolItemProxy.h b/Engine/Plugins/Experimental/SequenceNavigator/Source/SequenceNavigator/Public/Items/NavigationToolItemProxy.h
--- a/Engine/Plugins/Experimental/SequenceNavigator/Source/SequenceNavigator/Public/Items/NavigationToolItemProxy.h
+++ b/Engine/Plugins/Experimental/SequenceNavigator/Source/SequenceNavigator/Public/Items/NavigationToolItemProxy.h
@@ -2,41 +2,51 @@
 
 #pragma once
 
-#include "NavigationToolItem.h"
+#include "Items/INavigationToolItem.h"
+#include "Items/NavigationToolItem.h"
+#include "MVVM/ICastable.h"
+#include "NavigationToolDefines.h"
 
 #define UE_API SEQUENCENAVIGATOR_API
 
 namespace UE::SequenceNavigator
 {
 
+class INavigationTool;
+class INavigationToolView;
+
 /**
  * Item Proxies are Navigation Tool Items that with the sole purpose to group and hold common items together.
- * The description or name of such commonality between these Items should be the name of the Proxy that holds them.
+ * The description or name of such commonality between these items should be the name of the Proxy that holds them.
  * 
- * NOTE: Although Item Proxies by default require a parent to be visible in Navigation Tool,
+ * NOTE: Although Item proxies by default require a parent to be visible in Navigation Tool,
  * they can be created without a parent as a means to override behavior (e.g. DisplayName, Icon, etc)
  */
-class FNavigationToolItemProxy : public FNavigationToolItem
+class FNavigationToolItemProxy
+	: public FNavigationToolItem
 {
 public:
-	UE_NAVIGATIONTOOL_INHERITS_WITH_SUPER(FNavigationToolItemProxy, FNavigationToolItem);
+	UE_SEQUENCER_DECLARE_CASTABLE_API(UE_API, FNavigationToolItemProxy
+		, FNavigationToolItem)
 
-	UE_API FNavigationToolItemProxy(INavigationTool& InTool, const FNavigationToolItemPtr& InParentItem);
+	UE_API FNavigationToolItemProxy(INavigationTool& InTool, const FNavigationToolViewModelPtr& InParentItem);
 
 	//~ Begin INavigationToolItem
 	UE_API virtual bool IsItemValid() const override;
-	UE_API virtual void FindChildren(TArray<FNavigationToolItemPtr>& OutChildren, bool bRecursive) override final;
-	UE_API virtual void SetParent(FNavigationToolItemPtr InParent) override;
+	UE_API virtual void FindChildren(TArray<FNavigationToolViewModelWeakPtr>& OutWeakChildren
+		, const bool bInRecursive) override final;
+	UE_API virtual void SetParent(FNavigationToolViewModelPtr InParent) override;
 	UE_API virtual ENavigationToolItemViewMode GetSupportedViewModes(const INavigationToolView& InToolView) const override;
 	UE_API virtual bool CanAutoExpand() const override;
 	virtual FText GetClassName() const override { return FText::GetEmpty(); }
 	//~ End INavigationToolItem
 
 	uint32 GetPriority() const { return Priority; }
-	void SetPriority(uint32 InPriority) { Priority = InPriority; }
+	void SetPriority(const uint32 InPriority) { Priority = InPriority; }
 
-	/** Gets the Items that this Item Proxy is representing / holding (i.e. children) */
-	virtual void GetProxiedItems(const TSharedRef<INavigationToolItem>& InParent, TArray<FNavigationToolItemPtr>& OutChildren, bool bInRecursive) = 0;
+	/** Gets the items that this item proxy is representing / holding (ie. children) */
+	virtual void GetProxiedItems(const FNavigationToolViewModelPtr& InParent
+		, TArray<FNavigationToolViewModelWeakPtr>& OutWeakChildren, const bool bInRecursive) = 0;
 
 protected:
 	//~ Begin FNavigationToolItem
@@ -44,7 +54,8 @@ protected:
 	//~ End FNavigationToolItem
 
 private:
-	/** This Item Proxy's Order Priority (i.e. Highest priority is placed topmost or leftmost (depending on Orientation). Priority 0 is lowest priority */
+	/** This item proxy's order priority (ie. Highest priority is placed topmost or leftmost (depending on Orientation).
+	 * Priority 0 is the lowest priority. */
 	uint32 Priority = 0;
 };
 
