[SequenceNavigator] Refactor to use Sequencer view models and type macros

diff --git a/Engine/Plugins/Experimental/SequenceNavigator/Source/SequenceNavigator/Public/Items/NavigationToolItemUtils.h b/Engine/Plugins/Experimental/SequenceNavigator/Source/SequenceNavigator/Public/Items/NavigationToolItemUtils.h
--- a/Engine/Plugins/Experimental/SequenceNavigator/Source/SequenceNavigator/Public/Items/NavigationToolItemUtils.h
+++ b/Engine/Plugins/Experimental/SequenceNavigator/Source/SequenceNavigator/Public/Items/NavigationToolItemUtils.h
@@ -4,44 +4,50 @@
 
 #include "Algo/AllOf.h"
 #include "Containers/ContainersFwd.h"
-#include "Items/INavigationToolItem.h"
+#include "MVVM/ViewModels/ViewModelIterators.h"
 #include "NavigationToolDefines.h"
 #include "NavigationToolSequence.h"
 
+#define UE_API SEQUENCENAVIGATOR_API
+
 class ISequencer;
 class UMovieSceneMetaData;
 class UMovieSceneSubSection;
 
 namespace UE::SequenceNavigator
 {
-
 class INavigationTool;
 class FNavigationToolItem;
 class FNavigationToolSequence;
+}
+
+namespace UE::SequenceNavigator::ItemUtils
+{
 
 /**
  * Compares the absolute order of the items in the Navigation Tool and returns true if A comes before B in the Navigation Tool.
- * Useful to use when sorting Items.
+ * Useful to use when sorting items.
  */
-bool CompareToolItemOrder(const FNavigationToolItemPtr& InA, const FNavigationToolItemPtr& InB);
+bool CompareToolItemOrder(const FNavigationToolViewModelPtr& InA, const FNavigationToolViewModelPtr& InB);
 
 /** Returns two subset arrays of Items: one is containing only Sortable Items and the other Non Sortable Items */
-void SplitSortableAndUnsortableItems(const TArray<FNavigationToolItemPtr>& InItems
-	, TArray<FNavigationToolItemPtr>& OutSortable
-	, TArray<FNavigationToolItemPtr>& OutUnsortable);
+void SplitSortableAndUnsortableItems(const TArray<FNavigationToolViewModelWeakPtr>& InWeakItems
+	, TArray<FNavigationToolViewModelWeakPtr>& OutWeakSortable
+	, TArray<FNavigationToolViewModelWeakPtr>& OutWeakUnsortable);
 
 /**  */
-UMovieSceneSubSection* GetSequenceItemSubSection(const FNavigationToolItemPtr& InItem);
+UMovieSceneSubSection* GetSequenceItemSubSection(const FNavigationToolViewModelPtr& InItem);
 
 /**  */
-UMovieSceneMetaData* GetSequenceItemMetaData(const TSharedPtr<INavigationToolItem>& InItem);
+UE_API UMovieSceneMetaData* GetSequenceItemMetaData(const FNavigationToolViewModelPtr& InItem);
 
 /**
  * Removes the parent sequence prefix from the display name of a child sequence item.
  * @param InOutDisplayName The display name to modify, which will have the parent prefix removed if found.
  * @param InSequenceItem The child sequence item whose display name is being adjusted.
  */
-void RemoveSequenceDisplayNameParentPrefix(FText& InOutDisplayName, const TSharedRef<FNavigationToolSequence>& InSequenceItem);
+void RemoveSequenceDisplayNameParentPrefix(FText& InOutDisplayName
+	, const Sequencer::TViewModelPtr<FNavigationToolSequence>& InSequenceItem);
 
 /**
  * Appends the dirty symbol to the given display name.
@@ -71,24 +77,28 @@ ENUM_CLASS_FLAGS(ENavigationToolCompareState)
  *         ENavigationToolCompareState::PartialTrue if some items satisfy the condition
  */
 template<typename InItemType>
-ENavigationToolCompareState CompareChildrenItemState(const INavigationToolItem* const InItem
-	, const TFunctionRef<bool(const InItemType* const)>& InTrueFunction
-	, const TFunctionRef<bool(const InItemType* const)>& InFalseFunction
+ENavigationToolCompareState CompareChildrenItemState(const FNavigationToolViewModelPtr& InItem
+	, const TFunctionRef<bool(const Sequencer::TViewModelPtr<InItemType>)>& InTrueFunction
+	, const TFunctionRef<bool(const Sequencer::TViewModelPtr<InItemType>)>& InFalseFunction
 	, const ENavigationToolCompareState InDefaultState = ENavigationToolCompareState::NoneTrue)
 {
-	if (!InItem)
+	using namespace Sequencer;
+
+	if (!InItem.IsValid())
 	{
 		return InDefaultState;
 	}
 
-	const TArray<const InItemType*> Children = InItem->GetChildrenOfType<const InItemType>();
-	if (Children.IsEmpty())
+	const TViewModelListIterator<InItemType> WeakChildrenIter = InItem.AsModel()->GetChildrenOfType<InItemType>();
+	if (!WeakChildrenIter)
 	{
 		return InDefaultState;
 	}
 
-	const bool bAllMatch = Algo::AllOf(Children,
-		[&InTrueFunction](const InItemType* const InItem)
+	const TArray<TViewModelPtr<InItemType>> WeakChildren = WeakChildrenIter.ToArray();
+
+	const bool bAllMatch = Algo::AllOf(WeakChildren,
+		[&InTrueFunction](const TViewModelPtr<InItemType>& InItem)
 		{
 			return InTrueFunction(InItem);
 		});
@@ -97,8 +107,8 @@ ENavigationToolCompareState CompareChildrenItemState(const INavigationToolItem*
 		return ENavigationToolCompareState::AllTrue;
 	}
 
-	const bool bAllNoMatch = Algo::AllOf(Children,
-		[&InFalseFunction](const InItemType* const InItem)
+	const bool bAllNoMatch = Algo::AllOf(WeakChildren,
+		[&InFalseFunction](const Sequencer::TViewModelPtr<InItemType>& InItem)
 		{
 			return InFalseFunction(InItem);
 		});
@@ -120,17 +130,25 @@ ENavigationToolCompareState CompareChildrenItemState(const INavigationToolItem*
  *         ENavigationToolCompareState::PartialTrue if some items satisfy the condition
  */
 template<typename InItemType>
-ENavigationToolCompareState CompareChildrenItemStateSimple(const INavigationToolItem* const InItem
+ENavigationToolCompareState CompareChildrenItemStateSimple(const FNavigationToolViewModelPtr& InItem
 	, const TFunctionRef<bool(const InItemType* const)>& InTrueFunction)
 {
+	using namespace Sequencer;
+
 	if (!InItem)
 	{
 		return ENavigationToolCompareState::NoneTrue;
 	}
 
-	const TArray<const InItemType*> Children = InItem->GetChildrenOfType<const InItemType>();
+	const TViewModelListIterator<const InItemType> WeakChildrenIter = InItem.AsModel()->GetChildrenOfType<const InItemType>();
+	if (!WeakChildrenIter)
+	{
+		return ENavigationToolCompareState::NoneTrue;
+	}
+
+	const TArray<TViewModelPtr<const InItemType>> WeakChildren = WeakChildrenIter.ToArray();
 
-	const bool bAllMatch = Algo::AllOf(Children,
+	const bool bAllMatch = Algo::AllOf(WeakChildren,
 		[&InTrueFunction](const InItemType* const InItem)
 		{
 			return InTrueFunction(InItem);
@@ -140,7 +158,7 @@ ENavigationToolCompareState CompareChildrenItemStateSimple(const INavigationTool
 		return ENavigationToolCompareState::AllTrue;
 	}
 
-	const bool bAllNoMatch = Algo::AllOf(Children,
+	const bool bAllNoMatch = Algo::AllOf(WeakChildren,
 		[&InTrueFunction](const InItemType* const InItem)
 		{
 			return !InTrueFunction(InItem);
@@ -243,3 +261,5 @@ FSlateColor GetItemBindingColor(const ISequencer& InSequencer
 	, const FSlateColor& DefaultColor = FSlateColor::UseForeground());
 
 } // namespace UE::SequenceNavigator
+
+#undef UE_API
diff --git a/Engine/Plugins/Experimental/SequenceNavigator/Source/SequenceNavigator/Public/Items/NavigationToolMarker.h b/Engine/Plugins/Experimental/SequenceNavigator/Source/SequenceNavigator/Public/Items/NavigationToolMarker.h
--- a/Engine/Plugins/Experimental/SequenceNavigator/Source/SequenceNavigator/Public/Items/NavigationToolMarker.h
+++ b/Engine/Plugins/Experimental/SequenceNavigator/Source/SequenceNavigator/Public/Items/NavigationToolMarker.h
@@ -3,7 +3,8 @@
 #pragma once
 
 #include "Extensions/IInTimeExtension.h"
-#include "Extensions/IRenameableExtension.h"
+#include "MVVM/Extensions/IRenameableExtension.h"
+#include "MVVM/ICastable.h"
 #include "NavigationToolItem.h"
 
 #define UE_API SEQUENCENAVIGATOR_API
@@ -21,17 +22,17 @@ class FNavigationToolSequence;
  */
 class FNavigationToolMarker
 	: public FNavigationToolItem
-	, public IRenameableExtension
+	, public Sequencer::IRenameableExtension
 	, public IInTimeExtension
 {
 public:
-	UE_NAVIGATIONTOOL_INHERITS_WITH_SUPER(FNavigationToolMarker
+	UE_SEQUENCER_DECLARE_CASTABLE_API(UE_API, FNavigationToolMarker
 		, FNavigationToolItem
-		, IRenameableExtension
-		, IInTimeExtension);
+		, Sequencer::IRenameableExtension
+		, IInTimeExtension)
 
 	UE_API FNavigationToolMarker(INavigationTool& InTool
-		, const FNavigationToolItemPtr& InParentItem
+		, const FNavigationToolViewModelPtr& InParentItem
 		, const TSharedPtr<FNavigationToolSequence>& InParentSequenceItem
 		, const int32 InMarkedFrameIndex);
 
@@ -61,7 +62,7 @@ public:
 
 	//~ Begin IRenameableExtension
 	UE_API virtual bool CanRename() const override;
-	UE_API virtual bool Rename(const FString& InName) override;
+	UE_API virtual void Rename(const FText& InNewName) override;
 	//~ End IRenameableExtension
 
 	//~ Begin IInTimeExtension
@@ -81,7 +82,7 @@ protected:
 	UE_API UMovieSceneSequence* GetParentSequence() const;
 	UE_API UMovieScene* GetParentMovieScene() const;
 
-	TWeakPtr<FNavigationToolSequence> WeakParentSequenceItem;
+	Sequencer::TWeakViewModelPtr<FNavigationToolSequence> WeakParentSequenceItem;
 
 	int32 MarkedFrameIndex = INDEX_NONE;
 };
diff --git a/Engine/Plugins/Experimental/SequenceNavigator/Source/SequenceNavigator/Public/Items/NavigationToolSequence.h b/Engine/Plugins/Experimental/SequenceNavigator/Source/SequenceNavigator/Public/Items/NavigationToolSequence.h
--- a/Engine/Plugins/Experimental/SequenceNavigator/Source/SequenceNavigator/Public/Items/NavigationToolSequence.h
+++ b/Engine/Plugins/Experimental/SequenceNavigator/Source/SequenceNavigator/Public/Items/NavigationToolSequence.h
@@ -6,12 +6,14 @@
 #include "Extensions/IMarkerVisibilityExtension.h"
 #include "Extensions/IOutTimeExtension.h"
 #include "Extensions/IPlayheadExtension.h"
-#include "Extensions/IRenameableExtension.h"
 #include "Extensions/IRevisionControlExtension.h"
-#include "Extensions/ISequenceInactivableExtension.h"
-#include "Extensions/ISequenceLockableExtension.h"
+#include "MVVM/Extensions/IDeactivatableExtension.h"
+#include "MVVM/Extensions/ILockableExtension.h"
+#include "MVVM/Extensions/IRenameableExtension.h"
+#include "MVVM/ViewModelTypeID.h"
 #include "NavigationToolItem.h"
-#include "SequencerCoreFwd.h"
+
+#define UE_API SEQUENCENAVIGATOR_API
 
 class UMovieScene;
 class UMovieSceneSequence;
@@ -31,9 +33,9 @@ namespace UE::SequenceNavigator
  */
 class FNavigationToolSequence
 	: public FNavigationToolItem
-	, public IRenameableExtension
-	, public ISequenceInactivableExtension
-	, public ISequenceLockableExtension
+	, public Sequencer::IRenameableExtension
+	, public Sequencer::IDeactivatableExtension
+	, public Sequencer::ILockableExtension
 	, public IMarkerVisibilityExtension
 	, public IRevisionControlExtension
 	, public IPlayheadExtension
@@ -41,112 +43,112 @@ class FNavigationToolSequence
 	, public IOutTimeExtension
 {
 public:
-	UE_NAVIGATIONTOOL_INHERITS_WITH_SUPER(FNavigationToolSequence
+	UE_SEQUENCER_DECLARE_CASTABLE_API(UE_API, FNavigationToolSequence
 		, FNavigationToolItem
-		, IRenameableExtension
-		, ISequenceInactivableExtension
-		, ISequenceLockableExtension
+		, Sequencer::IRenameableExtension
+		, Sequencer::IDeactivatableExtension
+		, Sequencer::ILockableExtension
 		, IMarkerVisibilityExtension
 		, IRevisionControlExtension
 		, IPlayheadExtension
 		, IInTimeExtension
-		, IOutTimeExtension);
+		, IOutTimeExtension)
 
-	SEQUENCENAVIGATOR_API FNavigationToolSequence(INavigationTool& InTool
-		, const FNavigationToolItemPtr& InParentItem
+	UE_API FNavigationToolSequence(INavigationTool& InTool
+		, const FNavigationToolViewModelPtr& InParentItem
 		, UMovieSceneSequence* const InSequence
 		, UMovieSceneSubSection* const InSubSection
 		, const int32 InSubSectionIndex);
 
 	//~ Begin INavigationToolItem
 
-	SEQUENCENAVIGATOR_API virtual bool IsItemValid() const override;
-	SEQUENCENAVIGATOR_API virtual UObject* GetItemObject() const override;
-	SEQUENCENAVIGATOR_API virtual bool IsAllowedInTool() const override;
+	UE_API virtual bool IsItemValid() const override;
+	UE_API virtual UObject* GetItemObject() const override;
+	UE_API virtual bool IsAllowedInTool() const override;
 
-	SEQUENCENAVIGATOR_API virtual void FindChildren(TArray<FNavigationToolItemPtr>& OutChildren, const bool bInRecursive) override;
-	SEQUENCENAVIGATOR_API virtual void GetItemProxies(TArray<TSharedPtr<FNavigationToolItemProxy>>& OutItemProxies) override;
+	UE_API virtual void FindChildren(TArray<FNavigationToolViewModelWeakPtr>& OutWeakChildren, const bool bInRecursive) override;
+	UE_API virtual void GetItemProxies(TArray<TSharedPtr<FNavigationToolItemProxy>>& OutItemProxies) override;
 
-	SEQUENCENAVIGATOR_API virtual bool AddChild(const FNavigationToolAddItemParams& InAddItemParams) override;
-	SEQUENCENAVIGATOR_API virtual bool RemoveChild(const FNavigationToolRemoveItemParams& InRemoveItemParams) override;
+	UE_API virtual bool AddChild(const FNavigationToolAddItemParams& InAddItemParams) override;
+	UE_API virtual bool RemoveChild(const FNavigationToolRemoveItemParams& InRemoveItemParams) override;
 
-	SEQUENCENAVIGATOR_API virtual ENavigationToolItemViewMode GetSupportedViewModes(const INavigationToolView& InToolView) const override;
+	UE_API virtual ENavigationToolItemViewMode GetSupportedViewModes(const INavigationToolView& InToolView) const override;
 	virtual bool CanBeTopLevel() const override { return true; }
 
-	virtual FText GetDisplayName() const override;
-	SEQUENCENAVIGATOR_API virtual FText GetClassName() const override;
-	virtual FSlateIcon GetIcon() const override;
-	SEQUENCENAVIGATOR_API virtual FText GetIconTooltipText() const override;
+	UE_API virtual FText GetDisplayName() const override;
+	UE_API virtual FText GetClassName() const override;
+	UE_API virtual FSlateIcon GetIcon() const override;
+	UE_API virtual FText GetIconTooltipText() const override;
 
-	SEQUENCENAVIGATOR_API virtual bool IsSelected(const FNavigationToolScopedSelection& InSelection) const override;
-	SEQUENCENAVIGATOR_API virtual void Select(FNavigationToolScopedSelection& InSelection) const override;
+	UE_API virtual bool IsSelected(const FNavigationToolScopedSelection& InSelection) const override;
+	UE_API virtual void Select(FNavigationToolScopedSelection& InSelection) const override;
 
-	SEQUENCENAVIGATOR_API virtual void OnSelect() override;
-	SEQUENCENAVIGATOR_API virtual void OnDoubleClick() override;
+	UE_API virtual void OnSelect() override;
+	UE_API virtual void OnDoubleClick() override;
 
-	SEQUENCENAVIGATOR_API virtual void OnObjectsReplaced(const TMap<UObject*, UObject*>& InReplacementMap, bool bInRecursive) override;
+	UE_API virtual void OnObjectsReplaced(const TMap<UObject*, UObject*>& InReplacementMap, bool bInRecursive) override;
 
 	//~ End INavigationToolItem
 
-	//~ Begin IRenameableExtension
-	SEQUENCENAVIGATOR_API virtual bool CanRename() const override;
-	SEQUENCENAVIGATOR_API virtual bool Rename(const FString& InName) override;
-	//~ End IRenameableExtension
+	//~ Begin Sequencer::IRenameableExtension
+	UE_API virtual bool CanRename() const override;
+	UE_API virtual void Rename(const FText& InNewName) override;
+	//~ End Sequencer::IRenameableExtension
+
+	//~ Begin Sequencer::IDeactivatableExtension
+	UE_API virtual bool IsDeactivated() const override;
+	UE_API virtual void SetIsDeactivated(const bool bInIsDeactivated) override;
+	//~ End Sequencer::IDeactivatableExtension
 
-	//~ Begin ISequenceDeactivatableExtension
-	SEQUENCENAVIGATOR_API virtual EItemSequenceInactiveState GetInactiveState() const override;
-	SEQUENCENAVIGATOR_API virtual void SetIsInactive(const bool bInIsInactive) override;
-	//~ End ISequenceDeactivatableExtension
+	//~ Begin Sequencer::ILockableExtension
+	UE_API virtual Sequencer::ELockableLockState GetLockState() const override;
+	UE_API virtual void SetIsLocked(const bool bInIsLocked) override;
+	//~ End Sequencer::ILockableExtension
 
 	//~ Begin IMarkerVisibilityExtension
-	SEQUENCENAVIGATOR_API virtual EItemMarkerVisibility GetMarkerVisibility() const override;
-	SEQUENCENAVIGATOR_API virtual void SetMarkerVisibility(const bool bInVisible) override;
+	UE_API virtual EItemMarkerVisibility GetMarkerVisibility() const override;
+	UE_API virtual void SetMarkerVisibility(const bool bInVisible) override;
 	//~ End IMarkerVisibilityExtension
 
-	//~ Begin ISequenceLockableExtension
-	SEQUENCENAVIGATOR_API virtual EItemSequenceLockState GetLockState() const override;
-	SEQUENCENAVIGATOR_API virtual void SetIsLocked(const bool bInIsLocked) override;
-	//~ End ISequenceLockableExtension
-
 	//~ Begin IColorExtension
-	SEQUENCENAVIGATOR_API virtual TOptional<FColor> GetColor() const override;
-	SEQUENCENAVIGATOR_API virtual void SetColor(const TOptional<FColor>& InColor) override;
+	UE_API virtual TOptional<FColor> GetColor() const override;
+	UE_API virtual void SetColor(const TOptional<FColor>& InColor) override;
 	//~ End IColorExtension
 
 	//~ Begin IRevisionControlExtension
-	SEQUENCENAVIGATOR_API virtual EItemRevisionControlState GetRevisionControlState() const override;
-	SEQUENCENAVIGATOR_API virtual const FSlateBrush* GetRevisionControlStatusIcon() const override;
-	SEQUENCENAVIGATOR_API virtual FText GetRevisionControlStatusText() const override;
+	UE_API virtual EItemRevisionControlState GetRevisionControlState() const override;
+	UE_API virtual const FSlateBrush* GetRevisionControlStatusIcon() const override;
+	UE_API virtual FText GetRevisionControlStatusText() const override;
 	//~ End IRevisionControlExtension
 
 	//~ Begin IPlayheadExtension
-	SEQUENCENAVIGATOR_API virtual EItemContainsPlayhead ContainsPlayhead() const override;
+	UE_API virtual EItemContainsPlayhead ContainsPlayhead() const override;
 	//~ End IPlayheadExtension
 
 	//~ Begin IInTimeExtension
-	SEQUENCENAVIGATOR_API virtual FFrameNumber GetInTime() const override;
-	SEQUENCENAVIGATOR_API virtual void SetInTime(const FFrameNumber& InTime) override;
+	UE_API virtual FFrameNumber GetInTime() const override;
+	UE_API virtual void SetInTime(const FFrameNumber& InTime) override;
 	//~ End IInTimeExtension
 
 	//~ Begin IOutTimeExtension
-	SEQUENCENAVIGATOR_API virtual FFrameNumber GetOutTime() const override;
-	SEQUENCENAVIGATOR_API virtual void SetOutTime(const FFrameNumber& InTime) override;
+	UE_API virtual FFrameNumber GetOutTime() const override;
+	UE_API virtual void SetOutTime(const FFrameNumber& InTime) override;
 	//~ End IOutTimeExtension
 
-	SEQUENCENAVIGATOR_API UMovieSceneSequence* GetSequence() const;
+	UE_API UMovieSceneSequence* GetSequence() const;
 
-	SEQUENCENAVIGATOR_API UMovieSceneSubSection* GetSubSection() const;
-	SEQUENCENAVIGATOR_API int32 GetSubSectionIndex() const;
+	UE_API UMovieSceneSubSection* GetSubSection() const;
+	UE_API int32 GetSubSectionIndex() const;
 
-	SEQUENCENAVIGATOR_API UMovieScene* GetSequenceMovieScene() const;
+	UE_API UMovieScene* GetSequenceMovieScene() const;
 
-	SEQUENCENAVIGATOR_API TArray<FMovieSceneBinding> GetSortedBindings() const;
+	UE_API TArray<FMovieSceneBinding> GetSortedBindings() const;
 
-	UE::Sequencer::TViewModelPtr<UE::Sequencer::FSectionModel> GetViewModel() const;
+	Sequencer::TViewModelPtr<Sequencer::FSectionModel> GetViewModel() const;
 
 protected:
 	//~Begin FNavigationToolItem
-	SEQUENCENAVIGATOR_API virtual FNavigationToolItemId CalculateItemId() const override;
+	UE_API virtual FNavigationToolItemId CalculateItemId() const override;
 	//~End FNavigationToolItem
 
 	TWeakObjectPtr<UMovieSceneSubSection> WeakSubSection;
@@ -155,3 +157,5 @@ protected:
 };
 
 } // namespace UE::SequenceNavigator
+
+#undef UE_API
diff --git a/Engine/Plugins/Experimental/SequenceNavigator/Source/SequenceNavigator/Public/Items/NavigationToolSubTrack.h b/Engine/Plugins/Experimental/SequenceNavigator/Source/SequenceNavigator/Public/Items/NavigationToolSubTrack.h
--- a/Engine/Plugins/Experimental/SequenceNavigator/Source/SequenceNavigator/Public/Items/NavigationToolSubTrack.h
+++ b/Engine/Plugins/Experimental/SequenceNavigator/Source/SequenceNavigator/Public/Items/NavigationToolSubTrack.h
@@ -3,7 +3,8 @@
 #pragma once
 
 #include "Extensions/IMarkerVisibilityExtension.h"
-#include "Extensions/ISequenceInactivableExtension.h"
+#include "MVVM/Extensions/IDeactivatableExtension.h"
+#include "MVVM/ViewModelTypeID.h"
 #include "NavigationToolTrack.h"
 
 #define UE_API SEQUENCENAVIGATOR_API
@@ -17,34 +18,34 @@ namespace UE::SequenceNavigator
 
 class FNavigationToolSubTrack
 	: public FNavigationToolTrack
-	, public ISequenceInactivableExtension
+	, public Sequencer::IDeactivatableExtension
 	, public IMarkerVisibilityExtension
 {
 public:
-	UE_NAVIGATIONTOOL_INHERITS_WITH_SUPER(FNavigationToolSubTrack
+	UE_SEQUENCER_DECLARE_CASTABLE_API(UE_API, FNavigationToolSubTrack
 		, FNavigationToolTrack
-		, ISequenceInactivableExtension
-		, IMarkerVisibilityExtension);
+		, Sequencer::IDeactivatableExtension
+		, IMarkerVisibilityExtension)
 
 	UE_API FNavigationToolSubTrack(INavigationTool& InTool
-		, const FNavigationToolItemPtr& InParentItem
+		, const FNavigationToolViewModelPtr& InParentItem
 		, UMovieSceneSubTrack* const InSubTrack
 		, const TWeakObjectPtr<UMovieSceneSequence>& InSequence
 		, const TWeakObjectPtr<UMovieSceneSection>& InSection
         , const int32 InSubSectionIndex);
 
 	//~ Begin INavigationToolItem
-	UE_API virtual void FindChildren(TArray<FNavigationToolItemPtr>& OutChildren, const bool bInRecursive) override;
+	UE_API virtual void FindChildren(TArray<FNavigationToolViewModelWeakPtr>& OutWeakChildren, const bool bInRecursive) override;
 	UE_API virtual TOptional<EItemDropZone> CanAcceptDrop(const FDragDropEvent& InDragDropEvent, EItemDropZone InDropZone) override;
 	UE_API virtual FReply AcceptDrop(const FDragDropEvent& InDragDropEvent, EItemDropZone InDropZone) override;
 	//~ End INavigationToolItem
 
 	UE_API UMovieSceneSubTrack* GetSubTrack() const;
 
-	//~ Begin ISequenceDeactivatableExtension
-	UE_API virtual EItemSequenceInactiveState GetInactiveState() const override;
-	UE_API virtual void SetIsInactive(const bool bInIsInactive) override;
-	//~ End ISequenceDeactivatableExtension
+	//~ Begin Sequencer::IDeactivatableExtension
+	UE_API virtual bool IsDeactivated() const override;
+	UE_API virtual void SetIsDeactivated(const bool bInIsDeactivated) override;
+	//~ End Sequencer::IDeactivatableExtension
 
 	//~ Begin IMarkerVisibilityExtension
 	UE_API virtual EItemMarkerVisibility GetMarkerVisibility() const override;
diff --git a/Engine/Plugins/Experimental/SequenceNavigator/Source/SequenceNavigator/Public/Items/NavigationToolTrack.h b/Engine/Plugins/Experimental/SequenceNavigator/Source/SequenceNavigator/Public/Items/NavigationToolTrack.h
--- a/Engine/Plugins/Experimental/SequenceNavigator/Source/SequenceNavigator/Public/Items/NavigationToolTrack.h
+++ b/Engine/Plugins/Experimental/SequenceNavigator/Source/SequenceNavigator/Public/Items/NavigationToolTrack.h
@@ -3,12 +3,14 @@
 #pragma once
 
 #include "Extensions/IPlayheadExtension.h"
-#include "Extensions/ISequenceLockableExtension.h"
+#include "MVVM/Extensions/ILockableExtension.h"
+#include "MVVM/ViewModelTypeID.h"
 #include "NavigationToolItem.h"
-#include "SequencerCoreFwd.h"
 #include "Textures/SlateIcon.h"
 #include "UObject/WeakObjectPtrTemplates.h"
 
+#define UE_API SEQUENCENAVIGATOR_API
+
 class UMovieSceneSequence;
 class UMovieSceneSection;
 class UMovieSceneTrack;
@@ -23,17 +25,17 @@ namespace UE::SequenceNavigator
 
 class FNavigationToolTrack
 	: public FNavigationToolItem
-	, public ISequenceLockableExtension
+	, public Sequencer::ILockableExtension
 	, public IPlayheadExtension
 {
 public:
-	UE_NAVIGATIONTOOL_INHERITS_WITH_SUPER(FNavigationToolTrack
+	UE_SEQUENCER_DECLARE_CASTABLE_API(UE_API, FNavigationToolTrack
 		, FNavigationToolItem
-		, ISequenceLockableExtension
-		, IPlayheadExtension);
+		, Sequencer::ILockableExtension
+		, IPlayheadExtension)
 
 	FNavigationToolTrack(INavigationTool& InTool
-		, const FNavigationToolItemPtr& InParentItem
+		, const FNavigationToolViewModelPtr& InParentItem
 		, UMovieSceneTrack* const InTrack
 		, const TWeakObjectPtr<UMovieSceneSequence>& InSequence
 		, const TWeakObjectPtr<UMovieSceneSection>& InSection
@@ -45,7 +47,7 @@ public:
 	virtual UObject* GetItemObject() const override;
 	virtual bool IsAllowedInTool() const override;
 
-	virtual void FindChildren(TArray<FNavigationToolItemPtr>& OutChildren, const bool bInRecursive) override;
+	virtual void FindChildren(TArray<FNavigationToolViewModelWeakPtr>& OutWeakChildren, const bool bInRecursive) override;
 
 	virtual bool CanBeTopLevel() const override { return false; }
 
@@ -61,14 +63,14 @@ public:
 	virtual void OnSelect() override;
 	virtual void OnDoubleClick() override;
 
-	void OnObjectsReplaced(const TMap<UObject*, UObject*>& InReplacementMap, bool bInRecursive);
+	virtual void OnObjectsReplaced(const TMap<UObject*, UObject*>& InReplacementMap, bool bInRecursive) override;
 
 	//~ End INavigationToolItem
 
-	//~ Begin ISequenceLockableExtension
-	virtual EItemSequenceLockState GetLockState() const override;
+	//~ Begin ILockableExtension
+	virtual Sequencer::ELockableLockState GetLockState() const override;
 	virtual void SetIsLocked(const bool bInIsLocked) override;
-	//~ End ISequenceLockableExtension
+	//~ End ILockableExtension
 
 	//~ Begin IPlayheadExtension
 	virtual EItemContainsPlayhead ContainsPlayhead() const override;
@@ -98,3 +100,5 @@ protected:
 };
 
 } // namespace UE::SequenceNavigator
+
+#undef UE_API
diff --git a/Engine/Plugins/Experimental/SequenceNavigator/Source/SequenceNavigator/Public/Menus/NavigationToolItemMenuContext.h b/Engine/Plugins/Experimental/SequenceNavigator/Source/SequenceNavigator/Public/Menus/NavigationToolItemMenuContext.h
--- a/Engine/Plugins/Experimental/SequenceNavigator/Source/SequenceNavigator/Public/Menus/NavigationToolItemMenuContext.h
+++ b/Engine/Plugins/Experimental/SequenceNavigator/Source/SequenceNavigator/Public/Menus/NavigationToolItemMenuContext.h
@@ -22,10 +22,11 @@ class UNavigationToolItemMenuContext : public UObject
 	friend class FNavigationToolView;
 
 public:
-	void Init(const TWeakPtr<UE::SequenceNavigator::INavigationTool>& InTool, const TArray<UE::SequenceNavigator::FNavigationToolItemPtr>& InItemList)
+	void Init(const TWeakPtr<UE::SequenceNavigator::INavigationTool>& InTool
+		, const TArray<UE::Sequencer::TWeakViewModelPtr<UE::SequenceNavigator::INavigationToolItem>>& InWeakItemList)
 	{
 		WeakTool = InTool;
-		WeakItemList.Append(InItemList);
+		WeakItemList.Append(InWeakItemList);
 	}
 
 	TSharedPtr<UE::SequenceNavigator::INavigationTool> GetTool() const
@@ -33,7 +34,7 @@ public:
 		return WeakTool.Pin();
 	}
 
-	TConstArrayView<UE::SequenceNavigator::FNavigationToolItemWeakPtr> GetItems() const
+	TConstArrayView<UE::Sequencer::TWeakViewModelPtr<UE::SequenceNavigator::INavigationToolItem>> GetItems() const
 	{
 		return WeakItemList;
 	}
@@ -41,5 +42,5 @@ public:
 private:
 	TWeakPtr<UE::SequenceNavigator::INavigationTool> WeakTool;
 
-	TArray<UE::SequenceNavigator::FNavigationToolItemWeakPtr> WeakItemList;
+	TArray<UE::Sequencer::TWeakViewModelPtr<UE::SequenceNavigator::INavigationToolItem>> WeakItemList;
 };
diff --git a/Engine/Plugins/Experimental/SequenceNavigator/Source/SequenceNavigator/Public/NavigationToolDefines.h b/Engine/Plugins/Experimental/SequenceNavigator/Source/SequenceNavigator/Public/NavigationToolDefines.h
--- a/Engine/Plugins/Experimental/SequenceNavigator/Source/SequenceNavigator/Public/NavigationToolDefines.h
+++ b/Engine/Plugins/Experimental/SequenceNavigator/Source/SequenceNavigator/Public/NavigationToolDefines.h
@@ -3,20 +3,25 @@
 #pragma once
 
 #include "Containers/ContainersFwd.h"
+#include "MVVM/ViewModelPtr.h"
 #include "Templates/SharedPointer.h"
 #include "UObject/ObjectMacros.h"
 
+#define UE_API SEQUENCENAVIGATOR_API
+
 namespace UE::SequenceNavigator
 {
 
 class INavigationToolItem;
 
-using FNavigationToolItemRef = TSharedRef<INavigationToolItem>;
-using FNavigationToolItemPtr = TSharedPtr<INavigationToolItem>;
-using FNavigationToolItemWeakPtr = TWeakPtr<INavigationToolItem>;
+using FNavigationToolViewModelPtr = Sequencer::TViewModelPtr<INavigationToolItem>;
+using FNavigationToolViewModelWeakPtr = Sequencer::TWeakViewModelPtr<INavigationToolItem>;
+
+using FNavigationToolSharedPtr = TSharedPtr<Sequencer::TViewModelPtr<INavigationToolItem>::ViewModelType>;
+using FNavigationToolSharedRef = TSharedRef<Sequencer::TViewModelPtr<INavigationToolItem>::ViewModelType>;
 
-SEQUENCENAVIGATOR_API FName GetToolBarMenuName();
-SEQUENCENAVIGATOR_API FName GetItemContextMenuName();
+UE_API FName GetToolBarMenuName();
+UE_API FName GetItemContextMenuName();
 
 } // namespace UE::SequenceNavigator
 
@@ -65,3 +70,5 @@ enum class ENavigationToolDragDropActionType : uint8
 	Move,
 	Copy,
 };
+
+#undef UE_API
diff --git a/Engine/Plugins/Experimental/SequenceNavigator/Source/SequenceNavigator/Public/NavigationToolExtender.h b/Engine/Plugins/Experimental/SequenceNavigator/Source/SequenceNavigator/Public/NavigationToolExtender.h
--- a/Engine/Plugins/Experimental/SequenceNavigator/Source/SequenceNavigator/Public/NavigationToolExtender.h
+++ b/Engine/Plugins/Experimental/SequenceNavigator/Source/SequenceNavigator/Public/NavigationToolExtender.h
@@ -5,9 +5,10 @@
 #include "DragDropOps/NavigationToolItemDragDropOp.h"
 #include "Features/IModularFeature.h"
 #include "ItemProxies/NavigationToolItemProxyRegistry.h"
-#include "NavigationToolItemTypeId.h"
 #include "Providers/NavigationToolProvider.h"
 
+#define UE_API SEQUENCENAVIGATOR_API
+
 namespace UE::SequenceNavigator
 {
 
@@ -37,7 +38,7 @@ DECLARE_MULTICAST_DELEGATE_OneParam(FOnItemDragDropOpInitialized, FNavigationToo
 class FNavigationToolExtender : IModularFeature
 {
 public:
-	SEQUENCENAVIGATOR_API static FNavigationToolExtender& Get();
+	UE_API static FNavigationToolExtender& Get();
 
 	FNavigationToolExtender();
 	virtual ~FNavigationToolExtender();
@@ -51,11 +52,11 @@ public:
 	 * @param InSequencer The Sequencer instance for which the name is being retrieved.
 	 * @return The name of the Navigation Tool instance if available; otherwise, NAME_None.
 	 */
-	SEQUENCENAVIGATOR_API static FName GetToolInstanceId(const ISequencer& InSequencer);
-	SEQUENCENAVIGATOR_API static FName GetToolInstanceId(const INavigationTool& InTool);
+	UE_API static FName GetToolInstanceId(const ISequencer& InSequencer);
+	UE_API static FName GetToolInstanceId(const INavigationTool& InTool);
 
 	/** Find the Navigation Tool instance associated with the given Sequencer instance */
-	SEQUENCENAVIGATOR_API static TSharedPtr<INavigationTool> FindNavigationTool(const TSharedRef<ISequencer>& InSequencer);
+	UE_API static TSharedPtr<INavigationTool> FindNavigationTool(const TSharedRef<ISequencer>& InSequencer);
 
 	/**
 	 * Registers a Navigation Tool provider for a specified Navigation Tool instance.
@@ -65,7 +66,7 @@ public:
 	 * @param InProvider The Navigation Tool provider being registered.
 	 * @return True if the provider was successfully registered, false otherwise.
 	 */
-	SEQUENCENAVIGATOR_API static bool RegisterToolProvider(const TSharedRef<ISequencer>& InSequencer
+	UE_API static bool RegisterToolProvider(const TSharedRef<ISequencer>& InSequencer
 		, const TSharedRef<FNavigationToolProvider>& InProvider);
 
 	/**
@@ -74,10 +75,10 @@ public:
 	 * @param InProviderId The identifier of the provider to be unregistered.
 	 * @return True if the provider was successfully unregistered; false otherwise.
 	 */
-	SEQUENCENAVIGATOR_API static bool UnregisterToolProvider(const FName InToolId, const FName InProviderId);
+	UE_API static bool UnregisterToolProvider(const FName InToolId, const FName InProviderId);
 
 	/**  */
-	SEQUENCENAVIGATOR_API static TSharedPtr<FNavigationToolProvider> FindToolProvider(const FName InToolId, const FName InProviderId);
+	UE_API static TSharedPtr<FNavigationToolProvider> FindToolProvider(const FName InToolId, const FName InProviderId);
 
 	/**
 	 * Finds the tool providers associated with the given Navigation Tool instance.
@@ -85,7 +86,7 @@ public:
 	 * @param OutProviders A set that will be populated with the found tool providers as weak pointers.
 	 * @return True if tool providers were found and populated into OutProviders; otherwise, false.
 	 */
-	SEQUENCENAVIGATOR_API static bool FindToolProviders(const FName InToolId
+	UE_API static bool FindToolProviders(const FName InToolId
 		, TSet<TSharedRef<FNavigationToolProvider>>& OutProviders);
 
 	/**
@@ -94,7 +95,7 @@ public:
 	 * @param InPredicate A function that accepts a reference to an INavigationTool and a shared reference
 	 *                    to FNavigationToolProvider. Should return true to continue iteration or false to stop.
 	 */
-	SEQUENCENAVIGATOR_API static void ForEachProvider(const TFunction<bool(const FName InToolId
+	UE_API static void ForEachProvider(const TFunction<bool(const FName InToolId
 		, const TSharedRef<FNavigationToolProvider>& InProvider)>& InPredicate);
 
 	/**
@@ -105,7 +106,7 @@ public:
 	 * @param InPredicate A predicate function that takes a shared reference to a Navigation Tool provider as input.
 	 *                    The function should return true to continue iteration or false to stop it.
 	 */
-	SEQUENCENAVIGATOR_API static void ForEachToolProvider(const FName InToolId
+	UE_API static void ForEachToolProvider(const FName InToolId
 		, const TFunction<bool(const TSharedRef<FNavigationToolProvider>& InProvider)>& InPredicate);
 
 	/** Event called when a provider has been added or removed from a tool instance. */
@@ -130,7 +131,7 @@ public:
 	 * @param InItem The item to find the icon of
 	 * @return The custom icon for the item if a customization is found, FSlateIcon() otherwise
 	 */
-	SEQUENCENAVIGATOR_API static FSlateIcon FindOverrideIcon(TSharedPtr<const FNavigationToolItem> InItem);
+	UE_API static FSlateIcon FindOverrideIcon(const FNavigationToolViewModelPtr& InItem);
 
 	static bool AnyProviderSupportsSequencer(const FName InToolId, const ISequencer& InSequencer);
 
@@ -160,15 +161,15 @@ protected:
 	void OnSequencerActivated(const FName InToolId, FMovieSceneSequenceIDRef InSequenceId);
 	void OnSequencerClosed(const FName InToolId, const TSharedRef<ISequencer> InSequencer);
 
-	void RegisterOverridenIcon_Internal(const FNavigationToolItemTypeId& InItemTypeId, const TSharedRef<INavigationToolIconCustomization>& InIconCustomization);
-	void UnregisterOverriddenIcon_Internal(const FNavigationToolItemTypeId& InItemTypeId, const FName& InSpecializationIdentifier);
+	void RegisterOverridenIcon_Internal(const Sequencer::FViewModelTypeID& InItemTypeId, const TSharedRef<INavigationToolIconCustomization>& InIconCustomization);
+	void UnregisterOverriddenIcon_Internal(const Sequencer::FViewModelTypeID& InItemTypeId, const FName& InSpecializationIdentifier);
 
 	/**
 	 * Get the customization for the given item if any.
 	 * @param InItem Item to search the customization for.
 	 * @return The customization for the item if any; nullptr otherwise.
 	 */
-	TSharedPtr<INavigationToolIconCustomization> GetCustomizationForItem(const TSharedPtr<const FNavigationToolItem>& InItem) const;
+	TSharedPtr<INavigationToolIconCustomization> GetCustomizationForItem(const FNavigationToolViewModelPtr& InItem) const;
 
 	FDelegateHandle SequencerCreatedHandle;
 
@@ -177,7 +178,7 @@ protected:
 
 	/** Registered Sequencer types to Navigation Tool instances */
 	TMap<FName, FNavigationToolInstance> ToolInstances;
-	
+
 	FOnNavigationToolProvidersChanged ProvidersChangedDelegate;
 
 	FOnItemDragDropOpInitialized ItemDragDropOpInitializedDelegate;
@@ -188,7 +189,8 @@ protected:
 	struct FIconCustomizationKey
 	{
 		/** Navigation Tool item class FName (ex. NavigationToolItem) */
-		FNavigationToolItemTypeId ItemTypeId = FNavigationToolItemTypeId::Invalid();
+		//FNavigationToolItemTypeId ItemTypeId = FNavigationToolItemTypeId::Invalid();
+		uint32 ItemTypeId = 0;
 
 		/** Specialization identifier (ex. AvaShapeActor class FName for ActorIcon customization) */
 		FName CustomizationSpecializationIdentifier;
@@ -210,3 +212,5 @@ protected:
 };
 
 } // namespace UE::SequenceNavigator
+
+#undef UE_API
diff --git a/Engine/Plugins/Experimental/SequenceNavigator/Source/SequenceNavigator/Public/NavigationToolScopedSelection.h b/Engine/Plugins/Experimental/SequenceNavigator/Source/SequenceNavigator/Public/NavigationToolScopedSelection.h
--- a/Engine/Plugins/Experimental/SequenceNavigator/Source/SequenceNavigator/Public/NavigationToolScopedSelection.h
+++ b/Engine/Plugins/Experimental/SequenceNavigator/Source/SequenceNavigator/Public/NavigationToolScopedSelection.h
@@ -6,6 +6,8 @@
 #include "Containers/Set.h"
 #include "MovieSceneSequence.h"
 
+#define UE_API SEQUENCENAVIGATOR_API
+
 class ISequencer;
 class UMovieSceneSection;
 class UMovieSceneTrack;
@@ -31,20 +33,20 @@ public:
 
 	~FNavigationToolScopedSelection();
 
-	SEQUENCENAVIGATOR_API void Select(const FGuid& InObjectGuid);
-	SEQUENCENAVIGATOR_API void Select(UMovieSceneSection* const InSection);
-	SEQUENCENAVIGATOR_API void Select(UMovieSceneTrack* const InTrack);
-	SEQUENCENAVIGATOR_API void Select(UMovieSceneSequence* const InSequence
+	UE_API void Select(const FGuid& InObjectGuid);
+	UE_API void Select(UMovieSceneSection* const InSection);
+	UE_API void Select(UMovieSceneTrack* const InTrack);
+	UE_API void Select(UMovieSceneSequence* const InSequence
 		, const int32 InMarkedFrameIndex);
 
-	SEQUENCENAVIGATOR_API bool IsSelected(const UObject* const InObject) const;
-	SEQUENCENAVIGATOR_API bool IsSelected(const FGuid& InObjectGuid) const;
-	SEQUENCENAVIGATOR_API bool IsSelected(UMovieSceneSection* const InSection) const;
-	SEQUENCENAVIGATOR_API bool IsSelected(UMovieSceneTrack* const InTrack) const;
-	SEQUENCENAVIGATOR_API bool IsSelected(UMovieSceneSequence* const InSequence
+	UE_API bool IsSelected(const UObject* const InObject) const;
+	UE_API bool IsSelected(const FGuid& InObjectGuid) const;
+	UE_API bool IsSelected(UMovieSceneSection* const InSection) const;
+	UE_API bool IsSelected(UMovieSceneTrack* const InTrack) const;
+	UE_API bool IsSelected(UMovieSceneSequence* const InSequence
 		, const int32 InMarkedFrameIndex) const;
 
-	SEQUENCENAVIGATOR_API ISequencer& GetSequencer() const;
+	UE_API ISequencer& GetSequencer() const;
 
 private:
 	void SyncSelections();
@@ -63,3 +65,5 @@ private:
 };
 
 } // namespace UE::SequenceNavigator
+
+#undef UE_API
diff --git a/Engine/Plugins/Experimental/SequenceNavigator/Source/SequenceNavigator/Public/Providers/INavigationToolProvider.h b/Engine/Plugins/Experimental/SequenceNavigator/Source/SequenceNavigator/Public/Providers/INavigationToolProvider.h
--- a/Engine/Plugins/Experimental/SequenceNavigator/Source/SequenceNavigator/Public/Providers/INavigationToolProvider.h
+++ b/Engine/Plugins/Experimental/SequenceNavigator/Source/SequenceNavigator/Public/Providers/INavigationToolProvider.h
@@ -70,12 +70,12 @@ public:
 	 * Users can override this method to add or manipulate the child items of a given parent item.
 	 * @param InTool The Navigation Tool instance where the item resides.
 	 * @param InParentItem The item whose children are to be extended.
-	 * @param OutChildren The array to which additional child items will be added.
+	 * @param OutWeakChildren The array to which additional child items will be added.
 	 * @param bInRecursive Indicates whether to recursively find children for the provided item.
 	 */
 	virtual void OnExtendItemChildren(INavigationTool& InTool
-		, const FNavigationToolItemPtr& InParentItem
-		, TArray<FNavigationToolItemPtr>& OutChildren
+		, const FNavigationToolViewModelPtr& InParentItem
+		, TArray<FNavigationToolViewModelWeakPtr>& OutWeakChildren
 		, const bool bInRecursive) {}
 
 	/**
@@ -86,7 +86,7 @@ public:
 	 * @param OutItemProxies The array to which additional item proxies will be added.
 	 */
 	virtual void OnExtendItemProxiesForItem(INavigationTool& InTool
-		, const FNavigationToolItemPtr& InItem
+		, const FNavigationToolViewModelPtr& InItem
 		, TArray<TSharedPtr<FNavigationToolItemProxy>>& OutItemProxies) {}
 
 	/**
@@ -100,7 +100,7 @@ public:
 	virtual bool ShouldLockTool() const = 0;
 
 	/** An extended check to determine whether Item should be hidden in the Navigation Tool */
-	virtual bool ShouldHideItem(const FNavigationToolItemPtr& InItem) const = 0;
+	virtual bool ShouldHideItem(const FNavigationToolViewModelPtr& InItem) const = 0;
 
 	/** @return The name of the default column view to load when this provider is registered  */
 	virtual FText GetDefaultColumnView() const { return FText::GetEmpty(); }
@@ -108,20 +108,20 @@ public:
 	/** Determines whether an external Drag Drop event (i.e. not an Navigation Tool one) can be accepted by the Navigation Tool for a given Target Item */
 	virtual TOptional<EItemDropZone> OnToolItemCanAcceptDrop(const FDragDropEvent& InDragDropEvent
 		, const EItemDropZone InDropZone
-		, const FNavigationToolItemPtr& InTargetItem) const = 0;
+		, const FNavigationToolViewModelPtr& InTargetItem) const = 0;
 
 	/** Processes an external Drag Drop event (i.e. not an Navigation Tool one) for a given Target item */
 	virtual FReply OnToolItemAcceptDrop(const FDragDropEvent& InDragDropEvent
 		, const EItemDropZone InDropZone
-		, const FNavigationToolItemPtr& InTargetItem) = 0;
+		, const FNavigationToolViewModelPtr& InTargetItem) = 0;
 
 	virtual void UpdateItemIdContexts(const INavigationTool& InTool) = 0;
 
 	/** Event called when an Item has been renamed */
-	virtual void NotifyToolItemRenamed(const FNavigationToolItemPtr& InItem) {}
+	virtual void NotifyToolItemRenamed(const FNavigationToolViewModelPtr& InItem) {}
 
 	/** Event called when an Item has been deleted */
-	virtual void NotifyToolItemDeleted(const FNavigationToolItemPtr& InItem) {}
+	virtual void NotifyToolItemDeleted(const FNavigationToolViewModelPtr& InItem) {}
 };
 
 } // namespace UE::SequenceNavigator
diff --git a/Engine/Plugins/Experimental/SequenceNavigator/Source/SequenceNavigator/Public/Providers/NavigationToolProvider.h b/Engine/Plugins/Experimental/SequenceNavigator/Source/SequenceNavigator/Public/Providers/NavigationToolProvider.h
--- a/Engine/Plugins/Experimental/SequenceNavigator/Source/SequenceNavigator/Public/Providers/NavigationToolProvider.h
+++ b/Engine/Plugins/Experimental/SequenceNavigator/Source/SequenceNavigator/Public/Providers/NavigationToolProvider.h
@@ -5,6 +5,8 @@
 #include "NavigationToolDefines.h"
 #include "Providers/INavigationToolProvider.h"
 
+#define UE_API SEQUENCENAVIGATOR_API
+
 class UMovieSceneSequence;
 struct FNavigationToolBuiltInFilterParams;
 struct FNavigationToolSerializedTree;
@@ -22,16 +24,16 @@ class FNavigationToolProvider : public INavigationToolProvider
 public:
 	//~ Begin INavigationToolProvider
 
-	SEQUENCENAVIGATOR_API virtual void OnExtendColumnViews(TSet<FNavigationToolColumnView>& OutColumnViews) override;
-	SEQUENCENAVIGATOR_API virtual void OnExtendBuiltInFilters(TArray<FNavigationToolBuiltInFilterParams>& OutFilterParams) override;
+	UE_API virtual void OnExtendColumnViews(TSet<FNavigationToolColumnView>& OutColumnViews) override;
+	UE_API virtual void OnExtendBuiltInFilters(TArray<FNavigationToolBuiltInFilterParams>& OutFilterParams) override;
 
 	virtual bool ShouldLockTool() const override { return false; }
-	virtual bool ShouldHideItem(const FNavigationToolItemPtr& InItem) const override { return false; }
+	virtual bool ShouldHideItem(const FNavigationToolViewModelPtr& InItem) const override { return false; }
 
-	SEQUENCENAVIGATOR_API virtual void UpdateItemIdContexts(const INavigationTool& InTool) final override;
+	UE_API virtual void UpdateItemIdContexts(const INavigationTool& InTool) final override;
 
-	SEQUENCENAVIGATOR_API virtual TOptional<EItemDropZone> OnToolItemCanAcceptDrop(const FDragDropEvent& InDragDropEvent, EItemDropZone InDropZone, const FNavigationToolItemPtr& InTargetItem) const override;
-	SEQUENCENAVIGATOR_API virtual FReply OnToolItemAcceptDrop(const FDragDropEvent& InDragDropEvent, EItemDropZone InDropZone, const FNavigationToolItemPtr& InTargetItem) override;
+	UE_API virtual TOptional<EItemDropZone> OnToolItemCanAcceptDrop(const FDragDropEvent& InDragDropEvent, EItemDropZone InDropZone, const FNavigationToolViewModelPtr& InTargetItem) const override;
+	UE_API virtual FReply OnToolItemAcceptDrop(const FDragDropEvent& InDragDropEvent, EItemDropZone InDropZone, const FNavigationToolViewModelPtr& InTargetItem) override;
 
 	//~ End INavigationToolProvider
 
@@ -41,7 +43,7 @@ public:
 	/** Ensure saved editor only data contains an entry for a specific view, creating it if necessary */
 	void EnsureToolViewCount(const INavigationTool& InTool, const int32 InToolViewId);
 
-	SEQUENCENAVIGATOR_API bool IsSequenceSupported(UMovieSceneSequence* const InSequence) const;
+	UE_API bool IsSequenceSupported(UMovieSceneSequence* const InSequence) const;
 
 private:
 	friend class FNavigationToolExtender;
@@ -55,10 +57,10 @@ private:
 
 	void SaveSerializedTree(FNavigationTool& InTool, const bool bInResetTree);
 
-	void SaveSerializedTreeRecursive(const FNavigationToolItemPtr& InParentItem
+	void SaveSerializedTreeRecursive(const FNavigationToolViewModelPtr& InParentItem
 		, FNavigationToolSerializedTree& InSerializedTree);
 
-	void LoadSerializedTree(const FNavigationToolItemPtr& InParentItem
+	void LoadSerializedTree(const FNavigationToolViewModelPtr& InParentItem
 		, FNavigationToolSerializedTree* const InSerializedTree);
 
 	void CleanupExtendedColumnViews();
@@ -68,3 +70,5 @@ private:
 };
 
 } // namespace UE::SequenceNavigator
+
+#undef UE_API
diff --git a/Engine/Plugins/Experimental/SequenceNavigator/Source/SequenceNavigator/Public/Widgets/SNavigationToolTreeRow.h b/Engine/Plugins/Experimental/SequenceNavigator/Source/SequenceNavigator/Public/Widgets/SNavigationToolTreeRow.h
--- a/Engine/Plugins/Experimental/SequenceNavigator/Source/SequenceNavigator/Public/Widgets/SNavigationToolTreeRow.h
+++ b/Engine/Plugins/Experimental/SequenceNavigator/Source/SequenceNavigator/Public/Widgets/SNavigationToolTreeRow.h
@@ -2,6 +2,8 @@
 
 #pragma once
 
+#include "Items/NavigationToolItem.h"
+#include "MVVM/ViewModelPtr.h"
 #include "NavigationToolDefines.h"
 #include "Widgets/Views/STreeView.h"
 #include "Widgets/Views/STableRow.h"
@@ -9,11 +11,11 @@
 namespace UE::SequenceNavigator
 {
 
-class FNavigationToolItem;
 class FNavigationToolView;
 class SNavigationToolTreeView;
 
-class SNavigationToolTreeRow : public SMultiColumnTableRow<FNavigationToolItemPtr>
+class SNavigationToolTreeRow
+	: public SMultiColumnTableRow<FNavigationToolViewModelWeakPtr>
 {
 public:
 	SLATE_BEGIN_ARGS(SNavigationToolTreeRow) {}
@@ -23,14 +25,16 @@ public:
 	void Construct(const FArguments& InArgs
 		, const TSharedRef<FNavigationToolView>& InToolView
 		, const TSharedPtr<SNavigationToolTreeView>& InTreeView
-		, const FNavigationToolItemPtr& InItem);
+		, const FNavigationToolViewModelPtr& InItem);
 
+	//~ Begin SMultiColumnTableRow
 	virtual TSharedRef<SWidget> GenerateWidgetForColumn(const FName& InColumnName) override;
+	//~ End SMultiColumnTableRow
 
 	TAttribute<FText> GetHighlightText() const { return HighlightText; }
 
-	virtual FReply OnMouseButtonUp(const FGeometry& MyGeometry, const FPointerEvent& MouseEvent) override;
-	virtual FReply OnMouseButtonDoubleClick(const FGeometry& InGeometry, const FPointerEvent& InMouseEvent) override;
+	virtual FReply OnMouseButtonUp(const FGeometry& InGeometry, const FPointerEvent& InPointerEvent) override;
+	virtual FReply OnMouseButtonDoubleClick(const FGeometry& InGeometry, const FPointerEvent& InPointerEvent) override;
 
 	TSharedPtr<FNavigationToolView> GetToolView() const;
 
@@ -39,8 +43,8 @@ public:
 	/** The default reply if a row did not handle AcceptDrop */
 	FReply OnDefaultDrop(const FDragDropEvent& InDragDropEvent) const;
 
-private:
-	FNavigationToolItemPtr Item;
+protected:
+	FNavigationToolViewModelWeakPtr WeakItem;
 	
 	TWeakPtr<SNavigationToolTreeView> WeakTreeView;
 	
diff --git a/Engine/Plugins/Experimental/SequenceNavigator/Source/SequenceNavigator/SequenceNavigator.Build.cs b/Engine/Plugins/Experimental/SequenceNavigator/Source/SequenceNavigator/SequenceNavigator.Build.cs
--- a/Engine/Plugins/Experimental/SequenceNavigator/Source/SequenceNavigator/SequenceNavigator.Build.cs
+++ b/Engine/Plugins/Experimental/SequenceNavigator/Source/SequenceNavigator/SequenceNavigator.Build.cs
@@ -16,7 +16,8 @@ public class SequenceNavigator : ModuleRules
 
 		PublicDependencyModuleNames.AddRange(new string[]
 		{
-			"Core"
+			"Core",
+			"SequencerCore"
 		});
 
 		PrivateDependencyModuleNames.AddRange(new string[]
@@ -33,7 +34,6 @@ public class SequenceNavigator : ModuleRules
 			"MovieSceneTracks",
 			"Projects",
 			"Sequencer",
-			"SequencerCore",
 			"Slate",
 			"SlateCore",
 			"SourceControl",
