Sequencer: Fix failing test by capturing the ToolInstance

#rnx

diff --git a/Engine/Plugins/Experimental/SequenceNavigator/Source/SequenceNavigator/Private/NavigationToolExtender.cpp b/Engine/Plugins/Experimental/SequenceNavigator/Source/SequenceNavigator/Private/NavigationToolExtender.cpp
--- a/Engine/Plugins/Experimental/SequenceNavigator/Source/SequenceNavigator/Private/NavigationToolExtender.cpp
+++ b/Engine/Plugins/Experimental/SequenceNavigator/Source/SequenceNavigator/Private/NavigationToolExtender.cpp
@@ -423,14 +423,8 @@ void FNavigationToolExtender::AddSequencerToolBarExtension(const FNavigationTool
 	FToolMenuSection& NavigationToolSection = ToolMenu->FindOrAddSection(TEXT("NavigationTool"));
 
 	NavigationToolSection.AddDynamicEntry(TEXT("NavigationTool")
-		, FNewToolMenuSectionDelegate::CreateLambda([ToolId](FToolMenuSection& InSection)
+		, FNewToolMenuSectionDelegate::CreateLambda([ToolId, ToolInstance = InToolInstance](FToolMenuSection& InSection)
 		{
-			FNavigationToolInstance* const ToolInstance = Get().ToolInstances.Find(ToolId);
-			if (!ToolInstance)
-			{
-				return;
-			}
-
 			USequencerToolMenuContext* const ContextObject = InSection.FindContext<USequencerToolMenuContext>();
 			if (!ContextObject)
 			{
@@ -447,7 +441,7 @@ void FNavigationToolExtender::AddSequencerToolBarExtension(const FNavigationTool
 				FNavigationToolCommands::Get().ToggleToolTabVisible,
 				TAttribute<FText>(), TAttribute<FText>(),
 				FSlateIcon(FNavigationToolStyle::Get().GetStyleSetName(), TEXT("Icon.ToolBar")));
-			ToggleTabVisibleEntry.SetCommandList(ToolInstance->Instance->GetBaseCommandList());
+			ToggleTabVisibleEntry.SetCommandList(ToolInstance.Instance->GetBaseCommandList());
 			ToggleTabVisibleEntry.StyleNameOverride = TEXT("SequencerToolbar");
 			InSection.AddEntry(ToggleTabVisibleEntry);
 		}));
