[Sequencer] Update filter when keys added/removed

diff --git a/Engine/Plugins/Experimental/SequenceNavigator/Source/SequenceNavigator/Private/Filters/NavigationToolFilterData.cpp b/Engine/Plugins/Experimental/SequenceNavigator/Source/SequenceNavigator/Private/Filters/NavigationToolFilterData.cpp
--- a/Engine/Plugins/Experimental/SequenceNavigator/Source/SequenceNavigator/Private/Filters/NavigationToolFilterData.cpp
+++ b/Engine/Plugins/Experimental/SequenceNavigator/Source/SequenceNavigator/Private/Filters/NavigationToolFilterData.cpp
@@ -15,7 +15,9 @@ FNavigationToolFilterData::FNavigationToolFilterData(const FString& InRawFilterT
 
 bool FNavigationToolFilterData::operator==(const FNavigationToolFilterData& InRhs) const
 {
-	return ContainsFilterInNodes(InRhs) && TotalNodeCount == InRhs.GetTotalNodeCount();
+	return GetTotalNodeCount() == InRhs.GetTotalNodeCount()
+		&& GetDisplayNodeCount() == InRhs.GetDisplayNodeCount()
+		&& ContainsFilterInNodes(InRhs);
 }
 
 bool FNavigationToolFilterData::operator!=(const FNavigationToolFilterData& InRhs) const
diff --git a/Engine/Plugins/Experimental/SequenceNavigator/Source/SequenceNavigator/Private/NavigationToolView.cpp b/Engine/Plugins/Experimental/SequenceNavigator/Source/SequenceNavigator/Private/NavigationToolView.cpp
--- a/Engine/Plugins/Experimental/SequenceNavigator/Source/SequenceNavigator/Private/NavigationToolView.cpp
+++ b/Engine/Plugins/Experimental/SequenceNavigator/Source/SequenceNavigator/Private/NavigationToolView.cpp
@@ -2386,15 +2386,14 @@ bool FNavigationToolView::UpdateFilters()
 		return false;
 	}
 
-	const FNavigationToolFilterData& PreviousFilterData = FilterBar->GetFilterData();
+	const FNavigationToolFilterData PreviousFilterData = FilterBar->GetFilterData();
 	const FNavigationToolFilterData& FilterData = FilterBar->FilterNodes();
 
 	//bFilteringOnNodeGroups = Sequencer.GetFocusedMovieSceneSequence()->GetMovieScene()->GetNodeGroups().HasAnyActiveFilter();
 	bFilterUpdateRequested = false;
 
 	// Return whether the new list of FilteredNodes is different than the previous list
-	return (PreviousFilterData.GetDisplayNodeCount() != FilterData.GetDisplayNodeCount()
-		|| PreviousFilterData != FilterData);
+	return PreviousFilterData != FilterData;
 }
 
 } // namespace UE::SequenceNavigator
diff --git a/Engine/Source/Editor/Sequencer/Private/Filters/SequencerFilterData.cpp b/Engine/Source/Editor/Sequencer/Private/Filters/SequencerFilterData.cpp
--- a/Engine/Source/Editor/Sequencer/Private/Filters/SequencerFilterData.cpp
+++ b/Engine/Source/Editor/Sequencer/Private/Filters/SequencerFilterData.cpp
@@ -16,7 +16,9 @@ FSequencerFilterData::FSequencerFilterData(const FString& InRawFilterText)
 
 bool FSequencerFilterData::operator==(const FSequencerFilterData& InRhs) const
 {
-	return ContainsFilterInNodes(InRhs) && TotalNodeCount == InRhs.GetTotalNodeCount();
+	return GetTotalNodeCount() == InRhs.GetTotalNodeCount()
+		&& GetDisplayNodeCount() == InRhs.GetDisplayNodeCount()
+		&& ContainsFilterInNodes(InRhs);
 }
 
 bool FSequencerFilterData::operator!=(const FSequencerFilterData& InRhs) const
diff --git a/Engine/Source/Editor/Sequencer/Private/SequencerNodeTree.cpp b/Engine/Source/Editor/Sequencer/Private/SequencerNodeTree.cpp
--- a/Engine/Source/Editor/Sequencer/Private/SequencerNodeTree.cpp
+++ b/Engine/Source/Editor/Sequencer/Private/SequencerNodeTree.cpp
@@ -341,15 +341,14 @@ bool FSequencerNodeTree::UpdateFilters()
 		return false;
 	}
 
-	const FSequencerFilterData& PreviousFilterData = FilterBar->GetFilterData();
+	const FSequencerFilterData PreviousFilterData = FilterBar->GetFilterData();
 	const FSequencerFilterData& FilterData = FilterBar->FilterNodes();
 
 	bFilteringOnNodeGroups = Sequencer.GetFocusedMovieSceneSequence()->GetMovieScene()->GetNodeGroups().HasAnyActiveFilter();
 	bFilterUpdateRequested = false;
 
 	// Return whether the new list of FilteredNodes is different than the previous list
-	return (PreviousFilterData.GetDisplayNodeCount() != FilterData.GetDisplayNodeCount()
-		|| PreviousFilterData != FilterData);
+	return PreviousFilterData != FilterData;
 }
 
 void FSequencerNodeTree::CleanupMuteSolo(UMovieScene* MovieScene)
