[SequenceNavigator]
- Fix Motion Design sequence children from not showing
- Hide column widgets that are not supported by the item
- Add missing cpp includes for widgets used
- Refresh default column view on active provider changed
- Fix crash on editor close
- Remove commented out code

diff --git a/Engine/Plugins/Experimental/AvalancheSequenceNavigator/Source/AvalancheSequenceNavigator/Private/AvaNavigationToolProvider.cpp b/Engine/Plugins/Experimental/AvalancheSequenceNavigator/Source/AvalancheSequenceNavigator/Private/AvaNavigationToolProvider.cpp
--- a/Engine/Plugins/Experimental/AvalancheSequenceNavigator/Source/AvalancheSequenceNavigator/Private/AvaNavigationToolProvider.cpp
+++ b/Engine/Plugins/Experimental/AvalancheSequenceNavigator/Source/AvalancheSequenceNavigator/Private/AvaNavigationToolProvider.cpp
@@ -361,7 +361,6 @@ void FAvaNavigationToolProvider::ExtendToolToolBar()
 		, TAttribute<FText>(), TAttribute<FText>()
 		, FSlateIcon(FAppStyle::GetAppStyleSetName(), TEXT("AnimationEditor.SetKey")));
 	SequenceAddNewEntry.StyleNameOverride = SequencerToolbarStyleName;
-	SequenceAddNewEntry.InsertPosition = FToolMenuInsert(TEXT("AddFilter"), EToolMenuInsertType::Before);
 	SequenceAddNewEntry.SetCommandList(ToolCommands);
 	MotionDesignSection.AddEntry(SequenceAddNewEntry);
 
diff --git a/Engine/Plugins/Experimental/AvalancheSequenceNavigator/Source/AvalancheSequenceNavigator/Private/Items/NavigationToolAvaSequence.h b/Engine/Plugins/Experimental/AvalancheSequenceNavigator/Source/AvalancheSequenceNavigator/Private/Items/NavigationToolAvaSequence.h
--- a/Engine/Plugins/Experimental/AvalancheSequenceNavigator/Source/AvalancheSequenceNavigator/Private/Items/NavigationToolAvaSequence.h
+++ b/Engine/Plugins/Experimental/AvalancheSequenceNavigator/Source/AvalancheSequenceNavigator/Private/Items/NavigationToolAvaSequence.h
@@ -12,10 +12,12 @@ namespace UE::SequenceNavigator
 /**
  * Item in Navigation Tool representing a Motion Design Sequence.
  */
-class FNavigationToolAvaSequence : public FNavigationToolSequence
+class FNavigationToolAvaSequence
+	: public FNavigationToolSequence
 {
 public:
-	UE_SEQUENCER_DECLARE_VIEW_MODEL_TYPE_ID(FNavigationToolAvaSequence)
+	UE_SEQUENCER_DECLARE_CASTABLE(FNavigationToolAvaSequence
+		, FNavigationToolSequence)
 
 	FNavigationToolAvaSequence(INavigationTool& InTool
 		, const FNavigationToolViewModelPtr& InParentItem
diff --git a/Engine/Plugins/Experimental/SequenceNavigator/Source/SequenceNavigator/Private/Columns/NavigationToolColorColumn.cpp b/Engine/Plugins/Experimental/SequenceNavigator/Source/SequenceNavigator/Private/Columns/NavigationToolColorColumn.cpp
--- a/Engine/Plugins/Experimental/SequenceNavigator/Source/SequenceNavigator/Private/Columns/NavigationToolColorColumn.cpp
+++ b/Engine/Plugins/Experimental/SequenceNavigator/Source/SequenceNavigator/Private/Columns/NavigationToolColorColumn.cpp
@@ -1,10 +1,14 @@
 // Copyright Epic Games, Inc. All Rights Reserved.
 
 #include "Columns/NavigationToolColorColumn.h"
+#include "Extensions/IColorExtension.h"
 #include "INavigationToolView.h"
 #include "Items/NavigationToolSequence.h"
 #include "Widgets/Columns/SNavigationToolColor.h"
 #include "Widgets/Layout/SBox.h"
+#include "Widgets/SNavigationToolTreeRow.h"
+#include "Widgets/SNullWidget.h"
+#include "Widgets/Views/SHeaderRow.h"
 
 #define LOCTEXT_NAMESPACE "NavigationToolColorColumn"
 
@@ -32,7 +36,7 @@ SHeaderRow::FColumn::FArguments FNavigationToolColorColumn::ConstructHeaderRowCo
 		.DefaultLabel(GetColumnDisplayNameText())
 	    .DefaultTooltip(GetColumnDisplayNameText())
 		[
-			// Display nothing in the column header
+			// Display nothing in the column header instead of the default label
 			SNew(SBox)
 		];
 }
@@ -41,7 +45,7 @@ TSharedRef<SWidget> FNavigationToolColorColumn::ConstructRowWidget(const FNaviga
 	, const TSharedRef<INavigationToolView>& InView
 	, const TSharedRef<SNavigationToolTreeRow>& InRow)
 {
-	if (!InItem.AsModel()->IsA<FNavigationToolSequence>())
+	if (!InItem.AsModel()->IsA<IColorExtension>())
 	{
 		return SNullWidget::NullWidget;
 	}
diff --git a/Engine/Plugins/Experimental/SequenceNavigator/Source/SequenceNavigator/Private/Columns/NavigationToolCommentColumn.cpp b/Engine/Plugins/Experimental/SequenceNavigator/Source/SequenceNavigator/Private/Columns/NavigationToolCommentColumn.cpp
--- a/Engine/Plugins/Experimental/SequenceNavigator/Source/SequenceNavigator/Private/Columns/NavigationToolCommentColumn.cpp
+++ b/Engine/Plugins/Experimental/SequenceNavigator/Source/SequenceNavigator/Private/Columns/NavigationToolCommentColumn.cpp
@@ -3,6 +3,8 @@
 #include "Columns/NavigationToolCommentColumn.h"
 #include "INavigationToolView.h"
 #include "Widgets/Columns/SNavigationToolComment.h"
+#include "Widgets/SNavigationToolTreeRow.h"
+#include "Widgets/Views/SHeaderRow.h"
 
 #define LOCTEXT_NAMESPACE "NavigationToolCommentColumn"
 
diff --git a/Engine/Plugins/Experimental/SequenceNavigator/Source/SequenceNavigator/Private/Columns/NavigationToolDeactiveStateColumn.cpp b/Engine/Plugins/Experimental/SequenceNavigator/Source/SequenceNavigator/Private/Columns/NavigationToolDeactiveStateColumn.cpp
--- a/Engine/Plugins/Experimental/SequenceNavigator/Source/SequenceNavigator/Private/Columns/NavigationToolDeactiveStateColumn.cpp
+++ b/Engine/Plugins/Experimental/SequenceNavigator/Source/SequenceNavigator/Private/Columns/NavigationToolDeactiveStateColumn.cpp
@@ -2,11 +2,13 @@
 
 #include "Columns/NavigationToolDeactiveStateColumn.h"
 #include "INavigationToolView.h"
-#include "Items/INavigationToolItem.h"
 #include "Items/NavigationToolItem.h"
 #include "MVVM/Extensions/IDeactivatableExtension.h"
 #include "Widgets/Columns/SNavigationToolDeactiveState.h"
-#include "Widgets/SBoxPanel.h"
+#include "Widgets/Images/SImage.h"
+#include "Widgets/SNavigationToolTreeRow.h"
+#include "Widgets/Views/SHeaderRow.h"
+#include "Widgets/SNullWidget.h"
 
 #define LOCTEXT_NAMESPACE "NavigationToolDeactiveStateColumn"
 
@@ -50,15 +52,12 @@ TSharedRef<SWidget> FNavigationToolDeactiveStateColumn::ConstructRowWidget(const
 	, const TSharedRef<INavigationToolView>& InView
 	, const TSharedRef<SNavigationToolTreeRow>& InRow)
 {
-	using namespace Sequencer;
+	if (!InItem.AsModel()->IsA<Sequencer::IDeactivatableExtension>())
+	{
+		return SNullWidget::NullWidget;
+	}
 
-	return SNew(SHorizontalBox)
-		+ SHorizontalBox::Slot()
-		.AutoWidth()
-		.VAlign(VAlign_Center)
-		[
-			SNew(SNavigationToolDeactiveState, SharedThis(this), InItem.ImplicitCast(), InView, InRow)
-		];
+	return SNew(SNavigationToolDeactiveState, SharedThis(this), InItem.ImplicitCast(), InView, InRow);
 }
 
 } // namespace UE::SequenceNavigator
diff --git a/Engine/Plugins/Experimental/SequenceNavigator/Source/SequenceNavigator/Private/Columns/NavigationToolHBiasColumn.cpp b/Engine/Plugins/Experimental/SequenceNavigator/Source/SequenceNavigator/Private/Columns/NavigationToolHBiasColumn.cpp
--- a/Engine/Plugins/Experimental/SequenceNavigator/Source/SequenceNavigator/Private/Columns/NavigationToolHBiasColumn.cpp
+++ b/Engine/Plugins/Experimental/SequenceNavigator/Source/SequenceNavigator/Private/Columns/NavigationToolHBiasColumn.cpp
@@ -2,9 +2,9 @@
 
 #include "Columns/NavigationToolHBiasColumn.h"
 #include "INavigationToolView.h"
-#include "Items/INavigationToolItem.h"
-#include "Items/NavigationToolSequence.h"
 #include "Widgets/Columns/SNavigationToolHBias.h"
+#include "Widgets/SNavigationToolTreeRow.h"
+#include "Widgets/Views/SHeaderRow.h"
 
 #define LOCTEXT_NAMESPACE "NavigationToolHBiasColumn"
 
diff --git a/Engine/Plugins/Experimental/SequenceNavigator/Source/SequenceNavigator/Private/Columns/NavigationToolInTimeColumn.cpp b/Engine/Plugins/Experimental/SequenceNavigator/Source/SequenceNavigator/Private/Columns/NavigationToolInTimeColumn.cpp
--- a/Engine/Plugins/Experimental/SequenceNavigator/Source/SequenceNavigator/Private/Columns/NavigationToolInTimeColumn.cpp
+++ b/Engine/Plugins/Experimental/SequenceNavigator/Source/SequenceNavigator/Private/Columns/NavigationToolInTimeColumn.cpp
@@ -5,6 +5,9 @@
 #include "INavigationToolView.h"
 #include "Items/NavigationToolItem.h"
 #include "Widgets/Columns/SNavigationToolInTime.h"
+#include "Widgets/SNavigationToolTreeRow.h"
+#include "Widgets/SNullWidget.h"
+#include "Widgets/Views/SHeaderRow.h"
 
 #define LOCTEXT_NAMESPACE "NavigationToolInTimeColumn"
 
@@ -42,6 +45,11 @@ TSharedRef<SWidget> FNavigationToolInTimeColumn::ConstructRowWidget(const FNavig
 	, const TSharedRef<INavigationToolView>& InView
 	, const TSharedRef<SNavigationToolTreeRow>& InRow)
 {
+	if (!InItem.AsModel()->IsA<IInTimeExtension>())
+	{
+		return SNullWidget::NullWidget;
+	}
+
 	return SNew(SNavigationToolInTime, InItem.ImplicitCast(), InView, InRow);
 }
 
diff --git a/Engine/Plugins/Experimental/SequenceNavigator/Source/SequenceNavigator/Private/Columns/NavigationToolItemsColumn.cpp b/Engine/Plugins/Experimental/SequenceNavigator/Source/SequenceNavigator/Private/Columns/NavigationToolItemsColumn.cpp
--- a/Engine/Plugins/Experimental/SequenceNavigator/Source/SequenceNavigator/Private/Columns/NavigationToolItemsColumn.cpp
+++ b/Engine/Plugins/Experimental/SequenceNavigator/Source/SequenceNavigator/Private/Columns/NavigationToolItemsColumn.cpp
@@ -3,6 +3,8 @@
 #include "Columns/NavigationToolItemsColumn.h"
 #include "INavigationToolView.h"
 #include "Widgets/Columns/SNavigationToolItemList.h"
+#include "Widgets/SNavigationToolTreeRow.h"
+#include "Widgets/Views/SHeaderRow.h"
 
 #define LOCTEXT_NAMESPACE "NavigationToolItemsColumn"
 
diff --git a/Engine/Plugins/Experimental/SequenceNavigator/Source/SequenceNavigator/Private/Columns/NavigationToolLabelColumn.cpp b/Engine/Plugins/Experimental/SequenceNavigator/Source/SequenceNavigator/Private/Columns/NavigationToolLabelColumn.cpp
--- a/Engine/Plugins/Experimental/SequenceNavigator/Source/SequenceNavigator/Private/Columns/NavigationToolLabelColumn.cpp
+++ b/Engine/Plugins/Experimental/SequenceNavigator/Source/SequenceNavigator/Private/Columns/NavigationToolLabelColumn.cpp
@@ -3,6 +3,8 @@
 #include "Columns/NavigationToolLabelColumn.h"
 #include "INavigationToolView.h"
 #include "Items/NavigationToolItem.h"
+#include "Widgets/SNavigationToolTreeRow.h"
+#include "Widgets/Views/SHeaderRow.h"
 
 #define LOCTEXT_NAMESPACE "NavigationToolLabelColumn"
 
diff --git a/Engine/Plugins/Experimental/SequenceNavigator/Source/SequenceNavigator/Private/Columns/NavigationToolLengthColumn.cpp b/Engine/Plugins/Experimental/SequenceNavigator/Source/SequenceNavigator/Private/Columns/NavigationToolLengthColumn.cpp
--- a/Engine/Plugins/Experimental/SequenceNavigator/Source/SequenceNavigator/Private/Columns/NavigationToolLengthColumn.cpp
+++ b/Engine/Plugins/Experimental/SequenceNavigator/Source/SequenceNavigator/Private/Columns/NavigationToolLengthColumn.cpp
@@ -3,6 +3,8 @@
 #include "Columns/NavigationToolLengthColumn.h"
 #include "INavigationToolView.h"
 #include "Widgets/Columns/SNavigationToolLength.h"
+#include "Widgets/SNavigationToolTreeRow.h"
+#include "Widgets/Views/SHeaderRow.h"
 
 #define LOCTEXT_NAMESPACE "NavigationToolLengthColumn"
 
diff --git a/Engine/Plugins/Experimental/SequenceNavigator/Source/SequenceNavigator/Private/Columns/NavigationToolLockColumn.cpp b/Engine/Plugins/Experimental/SequenceNavigator/Source/SequenceNavigator/Private/Columns/NavigationToolLockColumn.cpp
--- a/Engine/Plugins/Experimental/SequenceNavigator/Source/SequenceNavigator/Private/Columns/NavigationToolLockColumn.cpp
+++ b/Engine/Plugins/Experimental/SequenceNavigator/Source/SequenceNavigator/Private/Columns/NavigationToolLockColumn.cpp
@@ -6,7 +6,10 @@
 #include "MVVM/Extensions/ILockableExtension.h"
 #include "Styling/StyleColors.h"
 #include "Widgets/Columns/SNavigationToolLock.h"
-#include "Widgets/SBoxPanel.h"
+#include "Widgets/Images/SImage.h"
+#include "Widgets/SNavigationToolTreeRow.h"
+#include "Widgets/SNullWidget.h"
+#include "Widgets/Views/SHeaderRow.h"
 
 #define LOCTEXT_NAMESPACE "NavigationToolLockColumn"
 
@@ -49,13 +52,12 @@ TSharedRef<SWidget> FNavigationToolLockColumn::ConstructRowWidget(const FNavigat
 	, const TSharedRef<INavigationToolView>& InView
 	, const TSharedRef<SNavigationToolTreeRow>& InRow)
 {
-	return SNew(SHorizontalBox)
-		+ SHorizontalBox::Slot()
-		.AutoWidth()
-		.VAlign(VAlign_Center)
-		[
-			SNew(SNavigationToolLock, InItem.ImplicitCast(), InView, InRow)
-		];
+	if (!InItem.AsModel()->IsA<Sequencer::ILockableExtension>())
+	{
+		return SNullWidget::NullWidget;
+	}
+
+	return SNew(SNavigationToolLock, InItem.ImplicitCast(), InView, InRow);
 }
 
 } // namespace UE::SequenceNavigator
diff --git a/Engine/Plugins/Experimental/SequenceNavigator/Source/SequenceNavigator/Private/Columns/NavigationToolMarkerVisibilityColumn.cpp b/Engine/Plugins/Experimental/SequenceNavigator/Source/SequenceNavigator/Private/Columns/NavigationToolMarkerVisibilityColumn.cpp
--- a/Engine/Plugins/Experimental/SequenceNavigator/Source/SequenceNavigator/Private/Columns/NavigationToolMarkerVisibilityColumn.cpp
+++ b/Engine/Plugins/Experimental/SequenceNavigator/Source/SequenceNavigator/Private/Columns/NavigationToolMarkerVisibilityColumn.cpp
@@ -5,9 +5,10 @@
 #include "INavigationToolView.h"
 #include "Items/NavigationToolItem.h"
 #include "Widgets/Columns/SNavigationToolMarkerVisibility.h"
-#include "Widgets/SBoxPanel.h"
-#include "Widgets/SOverlay.h"
-#include "Widgets/Text/STextBlock.h"
+#include "Widgets/Images/SImage.h"
+#include "Widgets/SNavigationToolTreeRow.h"
+#include "Widgets/SNullWidget.h"
+#include "Widgets/Views/SHeaderRow.h"
 
 #define LOCTEXT_NAMESPACE "NavigationToolMarkerVisibilityColumn"
 
@@ -50,6 +51,11 @@ TSharedRef<SWidget> FNavigationToolMarkerVisibilityColumn::ConstructRowWidget(co
 	, const TSharedRef<INavigationToolView>& InView
 	, const TSharedRef<SNavigationToolTreeRow>& InRow)
 {
+	if (!InItem.AsModel()->IsA<IMarkerVisibilityExtension>())
+	{
+		return SNullWidget::NullWidget;
+	}
+
 	return SNew(SNavigationToolMarkerVisibility, SharedThis(this), InItem.ImplicitCast(), InView, InRow);
 }
 
diff --git a/Engine/Plugins/Experimental/SequenceNavigator/Source/SequenceNavigator/Private/Columns/NavigationToolOutTimeColumn.cpp b/Engine/Plugins/Experimental/SequenceNavigator/Source/SequenceNavigator/Private/Columns/NavigationToolOutTimeColumn.cpp
--- a/Engine/Plugins/Experimental/SequenceNavigator/Source/SequenceNavigator/Private/Columns/NavigationToolOutTimeColumn.cpp
+++ b/Engine/Plugins/Experimental/SequenceNavigator/Source/SequenceNavigator/Private/Columns/NavigationToolOutTimeColumn.cpp
@@ -6,6 +6,8 @@
 #include "Items/NavigationToolItem.h"
 #include "Widgets/Columns/SNavigationToolOutTime.h"
 #include "Widgets/SNavigationToolTreeRow.h"
+#include "Widgets/SNullWidget.h"
+#include "Widgets/Views/SHeaderRow.h"
 
 #define LOCTEXT_NAMESPACE "NavigationToolOutTimeColumn"
 
@@ -43,6 +45,11 @@ TSharedRef<SWidget> FNavigationToolOutTimeColumn::ConstructRowWidget(const FNavi
 	, const TSharedRef<INavigationToolView>& InView
 	, const TSharedRef<SNavigationToolTreeRow>& InRow)
 {
+	if (!InItem.AsModel()->IsA<IOutTimeExtension>())
+	{
+		return SNullWidget::NullWidget;
+	}
+
 	return SNew(SNavigationToolOutTime, InItem.ImplicitCast(), InView, InRow);
 }
 
diff --git a/Engine/Plugins/Experimental/SequenceNavigator/Source/SequenceNavigator/Private/Columns/NavigationToolPlayheadColumn.cpp b/Engine/Plugins/Experimental/SequenceNavigator/Source/SequenceNavigator/Private/Columns/NavigationToolPlayheadColumn.cpp
--- a/Engine/Plugins/Experimental/SequenceNavigator/Source/SequenceNavigator/Private/Columns/NavigationToolPlayheadColumn.cpp
+++ b/Engine/Plugins/Experimental/SequenceNavigator/Source/SequenceNavigator/Private/Columns/NavigationToolPlayheadColumn.cpp
@@ -5,6 +5,9 @@
 #include "INavigationToolView.h"
 #include "Items/NavigationToolItem.h"
 #include "Widgets/Columns/SNavigationToolPlayhead.h"
+#include "Widgets/Images/SImage.h"
+#include "Widgets/SNavigationToolTreeRow.h"
+#include "Widgets/SNullWidget.h"
 #include "Widgets/Views/SHeaderRow.h"
 
 #define LOCTEXT_NAMESPACE "NavigationToolPlayheadColumn"
@@ -48,6 +51,11 @@ TSharedRef<SWidget> FNavigationToolPlayheadColumn::ConstructRowWidget(const FNav
 	, const TSharedRef<INavigationToolView>& InView
 	, const TSharedRef<SNavigationToolTreeRow>& InRow)
 {
+	if (!InItem.AsModel()->IsA<IPlayheadExtension>())
+	{
+		return SNullWidget::NullWidget;
+	}
+
 	return SNew(SNavigationToolPlayhead, SharedThis(this), InItem.ImplicitCast(), InView, InRow);
 }
 
diff --git a/Engine/Plugins/Experimental/SequenceNavigator/Source/SequenceNavigator/Private/Columns/NavigationToolRevisionControlColumn.cpp b/Engine/Plugins/Experimental/SequenceNavigator/Source/SequenceNavigator/Private/Columns/NavigationToolRevisionControlColumn.cpp
--- a/Engine/Plugins/Experimental/SequenceNavigator/Source/SequenceNavigator/Private/Columns/NavigationToolRevisionControlColumn.cpp
+++ b/Engine/Plugins/Experimental/SequenceNavigator/Source/SequenceNavigator/Private/Columns/NavigationToolRevisionControlColumn.cpp
@@ -1,12 +1,14 @@
 // Copyright Epic Games, Inc. All Rights Reserved.
 
 #include "Columns/NavigationToolRevisionControlColumn.h"
+#include "Extensions/IRevisionControlExtension.h"
 #include "INavigationToolView.h"
-#include "Items/INavigationToolItem.h"
-#include "Items/NavigationToolSequence.h"
 #include "RevisionControlStyle/RevisionControlStyle.h"
 #include "Widgets/Columns/SNavigationToolRevisionControl.h"
 #include "Widgets/Images/SImage.h"
+#include "Widgets/SNavigationToolTreeRow.h"
+#include "Widgets/SNullWidget.h"
+#include "Widgets/Views/SHeaderRow.h"
 
 #define LOCTEXT_NAMESPACE "NavigationToolRevisionControlColumn"
 
@@ -50,6 +52,11 @@ TSharedRef<SWidget> FNavigationToolRevisionControlColumn::ConstructRowWidget(con
 	, const TSharedRef<INavigationToolView>& InView
 	, const TSharedRef<SNavigationToolTreeRow>& InRow)
 {
+	if (!InItem.AsModel()->IsA<IRevisionControlExtension>())
+	{
+		return SNullWidget::NullWidget;
+	}
+
 	return SNew(SNavigationToolRevisionControl, SharedThis(this), InItem.ImplicitCast(), InView, InRow);
 }
 
diff --git a/Engine/Plugins/Experimental/SequenceNavigator/Source/SequenceNavigator/Private/Columns/NavigationToolStartFrameOffsetColumn.cpp b/Engine/Plugins/Experimental/SequenceNavigator/Source/SequenceNavigator/Private/Columns/NavigationToolStartFrameOffsetColumn.cpp
--- a/Engine/Plugins/Experimental/SequenceNavigator/Source/SequenceNavigator/Private/Columns/NavigationToolStartFrameOffsetColumn.cpp
+++ b/Engine/Plugins/Experimental/SequenceNavigator/Source/SequenceNavigator/Private/Columns/NavigationToolStartFrameOffsetColumn.cpp
@@ -3,6 +3,8 @@
 #include "Columns/NavigationToolStartFrameOffsetColumn.h"
 #include "INavigationToolView.h"
 #include "Widgets/Columns/SNavigationToolStartFrameOffset.h"
+#include "Widgets/SNavigationToolTreeRow.h"
+#include "Widgets/Views/SHeaderRow.h"
 
 #define LOCTEXT_NAMESPACE "NavigationToolStartFrameOffsetColumn"
 
diff --git a/Engine/Plugins/Experimental/SequenceNavigator/Source/SequenceNavigator/Private/Columns/NavigationToolTakeColumn.cpp b/Engine/Plugins/Experimental/SequenceNavigator/Source/SequenceNavigator/Private/Columns/NavigationToolTakeColumn.cpp
--- a/Engine/Plugins/Experimental/SequenceNavigator/Source/SequenceNavigator/Private/Columns/NavigationToolTakeColumn.cpp
+++ b/Engine/Plugins/Experimental/SequenceNavigator/Source/SequenceNavigator/Private/Columns/NavigationToolTakeColumn.cpp
@@ -3,8 +3,9 @@
 #include "Columns/NavigationToolTakeColumn.h"
 #include "Items/NavigationToolSequence.h"
 #include "INavigationToolView.h"
-#include "Items/INavigationToolItem.h"
 #include "Widgets/Columns/SNavigationToolTake.h"
+#include "Widgets/SNavigationToolTreeRow.h"
+#include "Widgets/Views/SHeaderRow.h"
 
 #define LOCTEXT_NAMESPACE "NavigationToolTakeColumn"
 
diff --git a/Engine/Plugins/Experimental/SequenceNavigator/Source/SequenceNavigator/Private/Items/NavigationToolItem.cpp b/Engine/Plugins/Experimental/SequenceNavigator/Source/SequenceNavigator/Private/Items/NavigationToolItem.cpp
--- a/Engine/Plugins/Experimental/SequenceNavigator/Source/SequenceNavigator/Private/Items/NavigationToolItem.cpp
+++ b/Engine/Plugins/Experimental/SequenceNavigator/Source/SequenceNavigator/Private/Items/NavigationToolItem.cpp
@@ -177,7 +177,6 @@ TArray<FNavigationToolViewModelPtr> FNavigationToolItem::FindPath(const TArray<F
 {
 	using namespace Sequencer;
 
-	//const TSharedPtr<const INavigationToolItem> This = SharedThis(this);
 	TArray<FNavigationToolViewModelPtr> Path;
 
 	for (const FNavigationToolViewModelWeakPtr& WeakItem : InWeakItems)
@@ -187,7 +186,6 @@ TArray<FNavigationToolViewModelPtr> FNavigationToolItem::FindPath(const TArray<F
 
 		while (CurrentItem.IsValid())
 		{
-			//if (This == *CurrentItem)
 			if (this == CurrentItem.Get())
 			{
 				Algo::Reverse(Path);
diff --git a/Engine/Plugins/Experimental/SequenceNavigator/Source/SequenceNavigator/Private/Items/NavigationToolSequence.cpp b/Engine/Plugins/Experimental/SequenceNavigator/Source/SequenceNavigator/Private/Items/NavigationToolSequence.cpp
--- a/Engine/Plugins/Experimental/SequenceNavigator/Source/SequenceNavigator/Private/Items/NavigationToolSequence.cpp
+++ b/Engine/Plugins/Experimental/SequenceNavigator/Source/SequenceNavigator/Private/Items/NavigationToolSequence.cpp
@@ -84,9 +84,14 @@ void FNavigationToolSequence::FindChildren(TArray<FNavigationToolViewModelWeakPt
 		return;
 	}
 
+	TArray<UMovieSceneTrack*> MovieSceneTracks = MovieScene->GetTracks();
+	MovieSceneTracks.Sort([](const UMovieSceneTrack& InA, const UMovieSceneTrack& InB)
+		{
+			return InA.GetSortingOrder() < InB.GetSortingOrder();
+		});
+
 	const TSharedRef<FNavigationToolProvider> ProviderRef = Provider.ToSharedRef();
 	const FNavigationToolViewModelPtr ParentItem = AsItemViewModel();
-	const TArray<UMovieSceneTrack*> MovieSceneTracks = MovieScene->GetTracks();
 
 	for (UMovieSceneTrack* const Track : MovieSceneTracks)
 	{
@@ -368,14 +373,12 @@ bool FNavigationToolSequence::IsDeactivated() const
 
 void FNavigationToolSequence::SetIsDeactivated(const bool bInIsDeactivated)
 {
-	const bool bNewActiveState = !bInIsDeactivated;
-
 	if (WeakSubSection.IsValid())
 	{
-		if (!WeakSubSection->IsReadOnly() && WeakSubSection->IsActive() != bNewActiveState)
+		if (!WeakSubSection->IsReadOnly() && WeakSubSection->IsActive() == bInIsDeactivated)
 		{
 			WeakSubSection->Modify();
-			WeakSubSection->SetIsActive(bNewActiveState);
+			WeakSubSection->SetIsActive(!bInIsDeactivated);
 		}
 	}
 
@@ -550,12 +553,13 @@ EItemContainsPlayhead FNavigationToolSequence::ContainsPlayhead() const
 	}
 
 	const FQualifiedFrameTime PlayheadTime = Sequencer->GetLocalTime();
+	const FFrameNumber PlayheadFrame = PlayheadTime.ConvertTo(Sequencer->GetFocusedTickResolution()).FloorToFrame();
 
 	if (Sequence == FocusedSequence)
 	{
 		if (UMovieScene* const SequenceMovieScene = Sequence->GetMovieScene())
 		{
-			return SequenceMovieScene->GetPlaybackRange().Contains(PlayheadTime.Time.FrameNumber)
+			return SequenceMovieScene->GetPlaybackRange().Contains(PlayheadFrame)
 				? EItemContainsPlayhead::ContainsPlayhead : EItemContainsPlayhead::None;
 		}
 	}
@@ -563,7 +567,7 @@ EItemContainsPlayhead FNavigationToolSequence::ContainsPlayhead() const
 	{
 		if (UMovieSceneSubSection* const SubSection = GetSubSection())
 		{
-			return SubSection->GetTrueRange().Contains(PlayheadTime.Time.FrameNumber)
+			return SubSection->GetRange().Contains(PlayheadTime.Time.FrameNumber)
 				? EItemContainsPlayhead::ContainsPlayhead : EItemContainsPlayhead::None;
 		}
 	}
diff --git a/Engine/Plugins/Experimental/SequenceNavigator/Source/SequenceNavigator/Private/NavigationTool.cpp b/Engine/Plugins/Experimental/SequenceNavigator/Source/SequenceNavigator/Private/NavigationTool.cpp
--- a/Engine/Plugins/Experimental/SequenceNavigator/Source/SequenceNavigator/Private/NavigationTool.cpp
+++ b/Engine/Plugins/Experimental/SequenceNavigator/Source/SequenceNavigator/Private/NavigationTool.cpp
@@ -70,32 +70,7 @@ void FNavigationTool::Init()
 			FIsActionChecked::CreateSP(this, &FNavigationTool::IsToolTabVisible));
 	}
 
-	const UNavigationToolSettings* const ToolSettings = GetDefault<UNavigationToolSettings>();
-
-	// We apply default views *after* all columns have been created for all providers
-	if (ToolSettings && ToolSettings->ShouldApplyDefaultColumnView())
-	{
-		TArray<FText> DefaultColumnViews;
-
-		ForEachProvider([&DefaultColumnViews](const TSharedRef<FNavigationToolProvider>& InProvider)
-			{
-				const FText DefaultColumnView = InProvider->GetDefaultColumnView();
-				if (!DefaultColumnView.IsEmptyOrWhitespace())
-				{
-					DefaultColumnViews.Add(DefaultColumnView);
-				}
-				return true;
-			});
-
-		// @TODO: priority int to better find a default priority view instead of just using the first index?
-		if (!DefaultColumnViews.IsEmpty())
-		{
-			ForEachToolView([&DefaultColumnViews](const TSharedPtr<FNavigationToolView>& InToolView)
-				{
-					InToolView->ApplyCustomColumnView(DefaultColumnViews[0]);
-				});
-		}
-	}
+	RefreshColumnView();
 
 	RefreshGlobalFilters();
 
@@ -257,6 +232,36 @@ TSharedPtr<ISequencer> FNavigationTool::GetSequencer() const
 	return WeakSequencer.Pin();
 }
 
+void FNavigationTool::RefreshColumnView()
+{
+	const UNavigationToolSettings* const ToolSettings = GetDefault<UNavigationToolSettings>();
+
+	// We apply default views *after* all columns have been created for all providers
+	if (ToolSettings && ToolSettings->ShouldApplyDefaultColumnView())
+	{
+		TArray<FText> DefaultColumnViews;
+
+		ForEachProvider([&DefaultColumnViews](const TSharedRef<FNavigationToolProvider>& InProvider)
+			{
+				const FText DefaultColumnView = InProvider->GetDefaultColumnView();
+				if (!DefaultColumnView.IsEmptyOrWhitespace())
+				{
+					DefaultColumnViews.Add(DefaultColumnView);
+				}
+				return true;
+			});
+
+		// @TODO: priority int to better find a default priority view instead of just using the first index?
+		if (!DefaultColumnViews.IsEmpty())
+		{
+			ForEachToolView([&DefaultColumnViews](const TSharedPtr<FNavigationToolView>& InToolView)
+				{
+					InToolView->ApplyCustomColumnView(DefaultColumnViews[0]);
+				});
+		}
+	}
+}
+
 bool FNavigationTool::DoesGlobalFilterExist(const FName InFilterId)
 {
 	for (const TSharedPtr<FNavigationToolBuiltInFilter>& GlobalFilter : GlobalFilters)
diff --git a/Engine/Plugins/Experimental/SequenceNavigator/Source/SequenceNavigator/Private/NavigationTool.h b/Engine/Plugins/Experimental/SequenceNavigator/Source/SequenceNavigator/Private/NavigationTool.h
--- a/Engine/Plugins/Experimental/SequenceNavigator/Source/SequenceNavigator/Private/NavigationTool.h
+++ b/Engine/Plugins/Experimental/SequenceNavigator/Source/SequenceNavigator/Private/NavigationTool.h
@@ -172,6 +172,9 @@ public:
 	/** Marks the Navigation Tool dirty. This triggers INavigationToolProvider::OnToolModified on next tick */
 	void SetToolModified();
 
+	/** If ShouldApplyDefaultColumnView(), applies the active providers default column view */
+	void RefreshColumnView();
+
 	bool DoesGlobalFilterExist(const FName InFilterId);
 	void RefreshGlobalFilters();
 
diff --git a/Engine/Plugins/Experimental/SequenceNavigator/Source/SequenceNavigator/Private/NavigationToolView.cpp b/Engine/Plugins/Experimental/SequenceNavigator/Source/SequenceNavigator/Private/NavigationToolView.cpp
--- a/Engine/Plugins/Experimental/SequenceNavigator/Source/SequenceNavigator/Private/NavigationToolView.cpp
+++ b/Engine/Plugins/Experimental/SequenceNavigator/Source/SequenceNavigator/Private/NavigationToolView.cpp
@@ -42,7 +42,7 @@ namespace UE::SequenceNavigator
 using namespace Sequencer;
 
 FNavigationToolView::FNavigationToolView(FPrivateToken)
-	: ToolSettings(GetMutableDefault<UNavigationToolSettings>())
+	: WeakToolSettings(GetMutableDefault<UNavigationToolSettings>())
 	, ItemContextMenu(MakeShared<FNavigationToolItemContextMenu>())
 {
 }
@@ -51,7 +51,10 @@ FNavigationToolView::~FNavigationToolView()
 {
 	if (UObjectInitialized())
 	{
-		ToolSettings->OnSettingChanged().RemoveAll(this);
+		if (UNavigationToolSettings* const ToolSettings = WeakToolSettings.Get())
+		{
+			ToolSettings->OnSettingChanged().RemoveAll(this);
+		}
 	}
 }
 
@@ -79,7 +82,10 @@ void FNavigationToolView::Init(const TSharedRef<FNavigationTool>& InTool, const
 			}
 		});
 
-	ToolSettings->OnSettingChanged().AddSP(this, &FNavigationToolView::OnToolSettingsChanged);
+	if (UNavigationToolSettings* const ToolSettings = WeakToolSettings.Get())
+	{
+		ToolSettings->OnSettingChanged().AddSP(this, &FNavigationToolView::OnToolSettingsChanged);
+	}
 
 	ToolViewWidget = SNew(SNavigationToolView, SharedThis(this));
 
@@ -117,8 +123,11 @@ void FNavigationToolView::CreateColumns(const TSharedRef<FNavigationToolProvider
 
 void FNavigationToolView::CreateDefaultColumnViews(const TSharedRef<FNavigationToolProvider>& InProvider)
 {
-	InProvider->OnExtendColumnViews(ToolSettings->GetCustomColumnViews());
-	ToolSettings->SaveConfig();
+	if (UNavigationToolSettings* const ToolSettings = WeakToolSettings.Get())
+	{
+		InProvider->OnExtendColumnViews(ToolSettings->GetCustomColumnViews());
+		ToolSettings->SaveConfig();
+	}
 }
 
 TSharedRef<FNavigationToolView> FNavigationToolView::CreateInstance(const int32 InToolViewId
@@ -594,7 +603,7 @@ void FNavigationToolView::GetChildrenOfItem(const FNavigationToolViewModelWeakPt
 
 			// For Muted Hierarchy to be in effect, not only does it have to be on
 			// but also the item should be shown (without counting the filter pass)
-			const bool bShouldUseMutedHierarchy = ToolSettings->ShouldUseMutedHierarchy();
+			const bool bShouldUseMutedHierarchy = WeakToolSettings.IsValid() ? WeakToolSettings.Get()->ShouldUseMutedHierarchy() : false;
 			const bool bShouldShowItemWithoutFilters = ShouldShowItem(WeakChildItem, false, InViewMode);
 			const bool bShouldMuteItem = bShouldUseMutedHierarchy && bShouldShowItemWithoutFilters;
 
@@ -792,9 +801,12 @@ void FNavigationToolView::NotifyItemSelectionChanged(const TArray<FNavigationToo
 	const TSharedPtr<FNavigationTool> Tool = WeakTool.Pin();
 	if (bInUpdateModeTools && Tool.IsValid())
 	{
-		if (ToolSettings->ShouldSyncSelectionToSequencer())
+		if (UNavigationToolSettings* const ToolSettings = WeakToolSettings.Get())
 		{
-			Tool->SyncSequencerSelection(WeakSelectedItems);
+			if (ToolSettings->ShouldSyncSelectionToSequencer())
+			{
+				Tool->SyncSequencerSelection(WeakSelectedItems);
+			}
 		}
 
 		Tool->SelectItems(WeakSelectedItems, ENavigationToolItemSelectionFlags::None);
@@ -1109,29 +1121,49 @@ bool FNavigationToolView::IsColumnVisible(const TSharedPtr<INavigationToolColumn
 
 ENavigationToolItemViewMode FNavigationToolView::GetItemDefaultViewMode() const
 {
-	return ToolSettings->GetItemDefaultViewMode();
+	if (UNavigationToolSettings* const ToolSettings = WeakToolSettings.Get())
+	{
+		return ToolSettings->GetItemDefaultViewMode();
+	}
+	return ENavigationToolItemViewMode::None;
 }
 
 ENavigationToolItemViewMode FNavigationToolView::GetItemProxyViewMode() const
 {
-	return ToolSettings->GetItemProxyViewMode();
+	if (UNavigationToolSettings* const ToolSettings = WeakToolSettings.Get())
+	{
+		return ToolSettings->GetItemProxyViewMode();
+	}
+	return ENavigationToolItemViewMode::None;
 }
 
 void FNavigationToolView::ToggleViewModeSupport(ENavigationToolItemViewMode& InOutViewMode, ENavigationToolItemViewMode InFlags)
 {
-	ToolSettings->ToggleViewModeSupport(InOutViewMode, InFlags);
+	if (UNavigationToolSettings* const ToolSettings = WeakToolSettings.Get())
+	{
+		ToolSettings->ToggleViewModeSupport(InOutViewMode, InFlags);
+	}
+
 	Refresh();
 }
 
 void FNavigationToolView::ToggleItemDefaultViewModeSupport(ENavigationToolItemViewMode InFlags)
 {
-	ToolSettings->ToggleItemDefaultViewModeSupport(InFlags);
+	if (UNavigationToolSettings* const ToolSettings = WeakToolSettings.Get())
+	{
+		ToolSettings->ToggleItemDefaultViewModeSupport(InFlags);
+	}
+
 	Refresh();
 }
 
 void FNavigationToolView::ToggleItemProxyViewModeSupport(ENavigationToolItemViewMode InFlags)
 {
-	ToolSettings->ToggleItemProxyViewModeSupport(InFlags);
+	if (UNavigationToolSettings* const ToolSettings = WeakToolSettings.Get())
+	{
+		ToolSettings->ToggleItemProxyViewModeSupport(InFlags);
+	}
+
 	Refresh();
 }
 
@@ -1154,45 +1186,77 @@ ECheckBoxState FNavigationToolView::GetViewModeCheckState(ENavigationToolItemVie
 
 ECheckBoxState FNavigationToolView::GetItemDefaultViewModeCheckState(ENavigationToolItemViewMode InFlags) const
 {
-	return GetViewModeCheckState(ToolSettings->GetItemDefaultViewMode(), InFlags);
+	if (UNavigationToolSettings* const ToolSettings = WeakToolSettings.Get())
+	{
+		return GetViewModeCheckState(ToolSettings->GetItemDefaultViewMode(), InFlags);
+	}
+	return ECheckBoxState::Unchecked;
 }
 
 ECheckBoxState FNavigationToolView::GetItemProxyViewModeCheckState(ENavigationToolItemViewMode InFlags) const
 {
-	return GetViewModeCheckState(ToolSettings->GetItemProxyViewMode(), InFlags);
+	if (UNavigationToolSettings* const ToolSettings = WeakToolSettings.Get())
+	{
+		return GetViewModeCheckState(ToolSettings->GetItemProxyViewMode(), InFlags);
+	}
+	return ECheckBoxState::Unchecked;
 }
 
 void FNavigationToolView::ToggleMutedHierarchy()
 {
-	ToolSettings->SetUseMutedHierarchy(!ToolSettings->ShouldUseMutedHierarchy());
+	if (UNavigationToolSettings* const ToolSettings = WeakToolSettings.Get())
+	{
+		ToolSettings->SetUseMutedHierarchy(!ToolSettings->ShouldUseMutedHierarchy());
+	}
+
 	Refresh();
 }
 
 bool FNavigationToolView::IsMutedHierarchyActive() const
 {
-	return ToolSettings->ShouldUseMutedHierarchy();
+	if (UNavigationToolSettings* const ToolSettings = WeakToolSettings.Get())
+	{
+		return ToolSettings->ShouldUseMutedHierarchy();
+	}
+	return false;
 }
 
 void FNavigationToolView::ToggleAutoExpandToSelection()
 {
-	ToolSettings->SetAutoExpandToSelection(!ToolSettings->ShouldAutoExpandToSelection());
+	if (UNavigationToolSettings* const ToolSettings = WeakToolSettings.Get())
+	{
+		ToolSettings->SetAutoExpandToSelection(!ToolSettings->ShouldAutoExpandToSelection());
+	}
+
 	Refresh();
 }
 
 bool FNavigationToolView::ShouldAutoExpandToSelection() const
 {
-	return ToolSettings->ShouldAutoExpandToSelection();
+	if (UNavigationToolSettings* const ToolSettings = WeakToolSettings.Get())
+	{
+		return ToolSettings->ShouldAutoExpandToSelection();
+	}
+	return false;
 }
 
 void FNavigationToolView::ToggleUseShortNames()
 {
-	ToolSettings->SetUseShortNames(!ToolSettings->ShouldUseShortNames());
+	if (UNavigationToolSettings* const ToolSettings = WeakToolSettings.Get())
+	{
+		ToolSettings->SetUseShortNames(!ToolSettings->ShouldUseShortNames());
+	}
+
 	Refresh();
 }
 
 bool FNavigationToolView::ShouldUseShortNames() const
 {
-	return ToolSettings->ShouldUseShortNames();
+	if (UNavigationToolSettings* const ToolSettings = WeakToolSettings.Get())
+	{
+		return ToolSettings->ShouldUseShortNames();
+	}
+	return false;
 }
 
 void FNavigationToolView::ToggleShowItemFilters()
@@ -1282,7 +1346,7 @@ FReply FNavigationToolView::OnDragDetected(const FGeometry& InGeometry
 	, const FPointerEvent& InPointerEvent
 	, const FNavigationToolViewModelWeakPtr InWeakTargetItem)
 {
-	if (!IsToolLocked())
+	if (IsToolLocked())
 	{
 		return FReply::Unhandled();
 	}
@@ -2175,6 +2239,12 @@ bool FNavigationToolView::CanResetAllColumnSizes() const
 
 void FNavigationToolView::SaveNewCustomColumnView()
 {
+	UNavigationToolSettings* const ToolSettings = WeakToolSettings.Get();
+	if (!ToolSettings)
+	{
+		return;
+	}
+
 	TSet<FNavigationToolColumnView>& CustomColumnViews = ToolSettings->GetCustomColumnViews();
 
 	SaveColumnState();
@@ -2245,6 +2315,12 @@ void FNavigationToolView::ApplyCustomColumnView(const FText InColumnViewName)
 		return;
 	}
 
+	UNavigationToolSettings* const ToolSettings = WeakToolSettings.Get();
+	if (!ToolSettings)
+	{
+		return;
+	}
+
 	const FNavigationToolColumnView* const SavedColumnView = ToolSettings->FindCustomColumnView(InColumnViewName);
 	if (!SavedColumnView)
 	{
diff --git a/Engine/Plugins/Experimental/SequenceNavigator/Source/SequenceNavigator/Private/NavigationToolView.h b/Engine/Plugins/Experimental/SequenceNavigator/Source/SequenceNavigator/Private/NavigationToolView.h
--- a/Engine/Plugins/Experimental/SequenceNavigator/Source/SequenceNavigator/Private/NavigationToolView.h
+++ b/Engine/Plugins/Experimental/SequenceNavigator/Source/SequenceNavigator/Private/NavigationToolView.h
@@ -392,8 +392,8 @@ protected:
 	/** Command List mapped to the View to handle things like Selected Items */
 	TSharedPtr<FUICommandList> ViewCommandList;
 
-	/**  */
-	const TObjectPtr<UNavigationToolSettings> ToolSettings;
+	/** Settings for the tool */
+	TWeakObjectPtr<UNavigationToolSettings> WeakToolSettings;
 
 	/** Root Items from Navigation Tool visible to this Instance. i.e. a Subset of all the Items in FNavigationTool */
 	TArray<FNavigationToolViewModelWeakPtr> WeakRootVisibleItems;
diff --git a/Engine/Plugins/Experimental/SequenceNavigator/Source/SequenceNavigator/Private/Providers/NavigationToolProvider.cpp b/Engine/Plugins/Experimental/SequenceNavigator/Source/SequenceNavigator/Private/Providers/NavigationToolProvider.cpp
--- a/Engine/Plugins/Experimental/SequenceNavigator/Source/SequenceNavigator/Private/Providers/NavigationToolProvider.cpp
+++ b/Engine/Plugins/Experimental/SequenceNavigator/Source/SequenceNavigator/Private/Providers/NavigationToolProvider.cpp
@@ -159,6 +159,8 @@ void FNavigationToolProvider::Activate(FNavigationTool& InTool)
 			InToolView->CreateDefaultColumnViews(SharedThisRef);
 		});
 
+	InTool.RefreshColumnView();
+
 	OnActivate();
 }
 
diff --git a/Engine/Plugins/Experimental/SequenceNavigator/Source/SequenceNavigator/Public/Items/NavigationToolItem.h b/Engine/Plugins/Experimental/SequenceNavigator/Source/SequenceNavigator/Public/Items/NavigationToolItem.h
--- a/Engine/Plugins/Experimental/SequenceNavigator/Source/SequenceNavigator/Public/Items/NavigationToolItem.h
+++ b/Engine/Plugins/Experimental/SequenceNavigator/Source/SequenceNavigator/Public/Items/NavigationToolItem.h
@@ -10,7 +10,6 @@
 #include "MVVM/Extensions/IOutlinerExtension.h"
 #include "MVVM/ViewModels/OutlinerItemModel.h"
 #include "NavigationToolSettings.h"
-#include "ToolkitBuilder.h"
 
 #define UE_API SEQUENCENAVIGATOR_API
 
@@ -114,8 +113,6 @@ public:
 
 	virtual TArray<FName> GetTags() const override { return TArray<FName>(); }
 
-	//virtual int32 GetItemTreeHeight() const override;
-
 	UE_API virtual TOptional<EItemDropZone> CanAcceptDrop(const FDragDropEvent& InDragDropEvent, const EItemDropZone InDropZone) override;
 	UE_API virtual FReply AcceptDrop(const FDragDropEvent& InDragDropEvent, const EItemDropZone InDropZone) override;
 
@@ -143,22 +140,12 @@ public:
 	{
 		const TSharedRef<FNavigationToolItem> SharedThisRef = SharedThis(this);
 		return Sequencer::CastViewModel<FNavigationToolItem>(SharedThisRef);
-		//return TOutlinerModelMixin::AsViewModel()->CastThisShared<FNavigationToolItem>();
 	}
 
 	Sequencer::TViewModelPtr<FNavigationToolItem> AsItemViewModelConst() const
 	{
-		//if (const FViewModel* const ViewModel = AsViewModel())
-		{
-			const TSharedRef<FNavigationToolItem> SharedThisRef = SharedThis(const_cast<FNavigationToolItem*>(this));
-			return Sequencer::CastViewModel<FNavigationToolItem>(SharedThisRef);
-		}
-		//return nullptr;
-/*
-		//TSharedRef<FNavigationToolItem> SharedThisRef = SharedThis(const_cast<FNavigationToolItem*>(this));
-		auto SharedThisRef = AsShared();
-		const FNavigationToolItem* const Item = SharedThisRef->CastThis<FNavigationToolItem>();
-		return Sequencer::TViewModelPtr(SharedThisRef, Item);// ? Item : CastDynamic<FNavigationToolItem>());*/
+		const TSharedRef<FNavigationToolItem> SharedThisRef = SharedThis(const_cast<FNavigationToolItem*>(this));
+		return Sequencer::CastViewModel<FNavigationToolItem>(SharedThisRef);
 	}
 
 protected:
