Sequencer: Add id column, extension

NavigationToolBinding and NavigationToolSequence implement IIdExtension by returning the sequence id or the object binding id

diff --git a/Engine/Plugins/Experimental/CinematicSequenceNavigator/Source/CinematicSequenceNavigator/Private/CinematicNavigationToolProvider.cpp b/Engine/Plugins/Experimental/CinematicSequenceNavigator/Source/CinematicSequenceNavigator/Private/CinematicNavigationToolProvider.cpp
--- a/Engine/Plugins/Experimental/CinematicSequenceNavigator/Source/CinematicSequenceNavigator/Private/CinematicNavigationToolProvider.cpp
+++ b/Engine/Plugins/Experimental/CinematicSequenceNavigator/Source/CinematicSequenceNavigator/Private/CinematicNavigationToolProvider.cpp
@@ -9,6 +9,7 @@
 #include "Columns/NavigationToolCommentColumn.h"
 #include "Columns/NavigationToolDeactiveStateColumn.h"
 #include "Columns/NavigationToolHBiasColumn.h"
+#include "Columns/NavigationToolIdColumn.h"
 #include "Columns/NavigationToolInTimeColumn.h"
 #include "Columns/NavigationToolItemsColumn.h"
 #include "Columns/NavigationToolLabelColumn.h"
@@ -137,6 +138,7 @@ void FCinematicNavigationToolProvider::OnExtendColumns(UE::SequenceNavigator::FN
 	OutExtender.AddColumn<FNavigationToolLockColumn>();
 	OutExtender.AddColumn<FNavigationToolColorColumn>();
 	OutExtender.AddColumn<FNavigationToolLabelColumn>();
+	OutExtender.AddColumn<FNavigationToolIdColumn>();
 	OutExtender.AddColumn<FNavigationToolItemsColumn>();
 	OutExtender.AddColumn<FNavigationToolInTimeColumn>();
 	OutExtender.AddColumn<FNavigationToolOutTimeColumn>();
diff --git a/Engine/Plugins/Experimental/LevelSequenceNavigatorBridge/Source/LevelSequenceNavigatorBridge/Private/LevelSequenceNavigationToolProvider.cpp b/Engine/Plugins/Experimental/LevelSequenceNavigatorBridge/Source/LevelSequenceNavigatorBridge/Private/LevelSequenceNavigationToolProvider.cpp
--- a/Engine/Plugins/Experimental/LevelSequenceNavigatorBridge/Source/LevelSequenceNavigatorBridge/Private/LevelSequenceNavigationToolProvider.cpp
+++ b/Engine/Plugins/Experimental/LevelSequenceNavigatorBridge/Source/LevelSequenceNavigatorBridge/Private/LevelSequenceNavigationToolProvider.cpp
@@ -1,10 +1,11 @@
-﻿// Copyright Epic Games, Inc. All Rights Reserved.
+// Copyright Epic Games, Inc. All Rights Reserved.
 
 #include "LevelSequenceNavigationToolProvider.h"
 #include "Columns/NavigationToolColorColumn.h"
 #include "Columns/NavigationToolColumnExtender.h"
 #include "Columns/NavigationToolCommentColumn.h"
 #include "Columns/NavigationToolHBiasColumn.h"
+#include "Columns/NavigationToolIdColumn.h"
 #include "Columns/NavigationToolInTimeColumn.h"
 #include "Columns/NavigationToolItemsColumn.h"
 #include "Columns/NavigationToolLabelColumn.h"
@@ -89,6 +90,7 @@ void FLevelSequenceNavigationToolProvider::OnExtendColumns(UE::SequenceNavigator
 	OutExtender.AddColumn<FNavigationToolLockColumn>();
 	OutExtender.AddColumn<FNavigationToolColorColumn>();
 	OutExtender.AddColumn<FNavigationToolLabelColumn>();
+	OutExtender.AddColumn<FNavigationToolIdColumn>();
 	OutExtender.AddColumn<FNavigationToolItemsColumn>();
 	OutExtender.AddColumn<FNavigationToolInTimeColumn>();
 	OutExtender.AddColumn<FNavigationToolOutTimeColumn>();
diff --git a/Engine/Plugins/Experimental/SequenceNavigator/Source/SequenceNavigator/Private/Columns/NavigationToolIdColumn.cpp b/Engine/Plugins/Experimental/SequenceNavigator/Source/SequenceNavigator/Private/Columns/NavigationToolIdColumn.cpp
new file mode 100644
--- /dev/null
+++ b/Engine/Plugins/Experimental/SequenceNavigator/Source/SequenceNavigator/Private/Columns/NavigationToolIdColumn.cpp
@@ -0,0 +1,57 @@
+// Copyright Epic Games, Inc. All Rights Reserved.
+
+#include "Columns/NavigationToolIdColumn.h"
+#include "Extensions/IIdExtension.h"
+#include "Items/NavigationToolSequence.h"
+#include "INavigationToolView.h"
+#include "Widgets/Columns/SNavigationToolId.h"
+#include "Widgets/SNavigationToolTreeRow.h"
+#include "Widgets/Views/SHeaderRow.h"
+
+#define LOCTEXT_NAMESPACE "NavigationToolIdColumn"
+
+namespace UE::SequenceNavigator
+{
+
+UE_SEQUENCER_DEFINE_CASTABLE(FNavigationToolIdColumn)
+
+FText FNavigationToolIdColumn::GetColumnDisplayNameText() const
+{
+	return LOCTEXT("IdColumn", "Id");
+}
+
+const FSlateBrush* FNavigationToolIdColumn::GetIconBrush() const
+{
+	return FAppStyle::GetBrush(TEXT("ClassIcon.LevelSequence"));
+}
+
+SHeaderRow::FColumn::FArguments FNavigationToolIdColumn::ConstructHeaderRowColumn(const TSharedRef<INavigationToolView>& InView
+	, const float InFillSize)
+{
+	const FName ColumnId = GetColumnId();
+
+	return SHeaderRow::Column(ColumnId)
+		.FillWidth(InFillSize)
+		.HAlignHeader(HAlign_Center)
+		.VAlignHeader(VAlign_Center)
+		.HAlignCell(HAlign_Fill)
+		.VAlignCell(VAlign_Center)
+		.DefaultLabel(GetColumnDisplayNameText())
+		.OnGetMenuContent(InView, &INavigationToolView::GetColumnMenuContent, ColumnId);
+}
+
+TSharedRef<SWidget> FNavigationToolIdColumn::ConstructRowWidget(const FNavigationToolViewModelPtr& InItem
+	, const TSharedRef<INavigationToolView>& InView
+	, const TSharedRef<SNavigationToolTreeRow>& InRow)
+{
+	if (!InItem.AsModel()->IsA<IIdExtension>())
+	{
+		return SNullWidget::NullWidget;
+	}
+
+	return SNew(SNavigationToolId, InItem.ImplicitCast(), InView, InRow);
+}
+
+} // namespace UE::SequenceNavigator
+
+#undef LOCTEXT_NAMESPACE
diff --git a/Engine/Plugins/Experimental/SequenceNavigator/Source/SequenceNavigator/Private/Items/NavigationToolBinding.cpp b/Engine/Plugins/Experimental/SequenceNavigator/Source/SequenceNavigator/Private/Items/NavigationToolBinding.cpp
--- a/Engine/Plugins/Experimental/SequenceNavigator/Source/SequenceNavigator/Private/Items/NavigationToolBinding.cpp
+++ b/Engine/Plugins/Experimental/SequenceNavigator/Source/SequenceNavigator/Private/Items/NavigationToolBinding.cpp
@@ -255,6 +255,11 @@ void FNavigationToolBinding::SetIsLocked(const bool bInIsLocked)
 	}
 }*/
 
+FText FNavigationToolBinding::GetId() const
+{
+	return FText::FromString(Binding.GetObjectGuid().ToString());
+}
+
 EItemContainsPlayhead FNavigationToolBinding::ContainsPlayhead() const
 {
 	const TSharedPtr<ISequencer> Sequencer = Tool.GetSequencer();
diff --git a/Engine/Plugins/Experimental/SequenceNavigator/Source/SequenceNavigator/Private/Items/NavigationToolSequence.cpp b/Engine/Plugins/Experimental/SequenceNavigator/Source/SequenceNavigator/Private/Items/NavigationToolSequence.cpp
--- a/Engine/Plugins/Experimental/SequenceNavigator/Source/SequenceNavigator/Private/Items/NavigationToolSequence.cpp
+++ b/Engine/Plugins/Experimental/SequenceNavigator/Source/SequenceNavigator/Private/Items/NavigationToolSequence.cpp
@@ -1,7 +1,10 @@
 // Copyright Epic Games, Inc. All Rights Reserved.
 
 #include "Items/NavigationToolSequence.h"
+#include "Compilation/MovieSceneCompiledDataManager.h"
 #include "GameFramework/Actor.h"
+#include "Evaluation/MovieSceneEvaluationTemplateInstance.h"
+#include "Evaluation/MovieSceneSequenceHierarchy.h"
 #include "INavigationTool.h"
 #include "Items/NavigationToolActor.h"
 #include "Items/NavigationToolItemUtils.h"
@@ -497,6 +500,36 @@ void FNavigationToolSequence::SetColor(const TOptional<FColor>& InColor)
 	WeakSubSection->SetColorTint(InColor.Get(FColor()));
 }
 
+FText FNavigationToolSequence::GetId() const
+{
+	UMovieSceneSequence* const Sequence = GetSequence();
+	if (!Sequence)
+	{
+		return FText();
+	}
+
+	const TSharedPtr<ISequencer> Sequencer = Tool.GetSequencer();
+	if (!Sequencer.IsValid())
+	{
+		return FText();
+	}
+
+	FMovieSceneRootEvaluationTemplateInstance& RootInstance = Sequencer->GetEvaluationTemplate();
+	const FMovieSceneSequenceHierarchy* Hierarchy = RootInstance.GetCompiledDataManager()->FindHierarchy(RootInstance.GetCompiledDataID());
+	if (Hierarchy)
+	{
+		for (const TTuple<FMovieSceneSequenceID, FMovieSceneSubSequenceData>& Pair : Hierarchy->AllSubSequenceData())
+		{
+			if (Sequence == Pair.Value.GetSequence())
+			{
+				return FText::FromString(FString::Printf(TEXT("%d"), Pair.Key.GetInternalValue()));
+			}
+		}
+	}
+
+	return FText();
+}
+
 EItemRevisionControlState FNavigationToolSequence::GetRevisionControlState() const
 {
 	UMovieSceneSequence* const Sequence = GetSequence();
diff --git a/Engine/Plugins/Experimental/SequenceNavigator/Source/SequenceNavigator/Private/NavigationToolViewModelTypeIDs.cpp b/Engine/Plugins/Experimental/SequenceNavigator/Source/SequenceNavigator/Private/NavigationToolViewModelTypeIDs.cpp
--- a/Engine/Plugins/Experimental/SequenceNavigator/Source/SequenceNavigator/Private/NavigationToolViewModelTypeIDs.cpp
+++ b/Engine/Plugins/Experimental/SequenceNavigator/Source/SequenceNavigator/Private/NavigationToolViewModelTypeIDs.cpp
@@ -2,6 +2,7 @@
 
 #include "Columns/INavigationToolColumn.h"
 #include "Extensions/IColorExtension.h"
+#include "Extensions/IIdExtension.h"
 #include "Extensions/IInTimeExtension.h"
 #include "Extensions/IMarkerVisibilityExtension.h"
 #include "Extensions/IOutTimeExtension.h"
@@ -21,6 +22,7 @@ UE_SEQUENCER_DEFINE_VIEW_MODEL_TYPE_ID(INavigationToolColumn)
 UE_SEQUENCER_DEFINE_CASTABLE(INavigationToolItemAction)
 
 UE_SEQUENCER_DEFINE_VIEW_MODEL_TYPE_ID(IColorExtension)
+UE_SEQUENCER_DEFINE_VIEW_MODEL_TYPE_ID(IIdExtension)
 UE_SEQUENCER_DEFINE_VIEW_MODEL_TYPE_ID(IInTimeExtension)
 UE_SEQUENCER_DEFINE_VIEW_MODEL_TYPE_ID(IMarkerVisibilityExtension)
 UE_SEQUENCER_DEFINE_VIEW_MODEL_TYPE_ID(IOutTimeExtension)
diff --git a/Engine/Plugins/Experimental/SequenceNavigator/Source/SequenceNavigator/Private/Widgets/Columns/SNavigationToolId.cpp b/Engine/Plugins/Experimental/SequenceNavigator/Source/SequenceNavigator/Private/Widgets/Columns/SNavigationToolId.cpp
new file mode 100644
--- /dev/null
+++ b/Engine/Plugins/Experimental/SequenceNavigator/Source/SequenceNavigator/Private/Widgets/Columns/SNavigationToolId.cpp
@@ -0,0 +1,47 @@
+// Copyright Epic Games, Inc. All Rights Reserved.
+
+#include "SNavigationToolId.h"
+#include "Extensions/IIdExtension.h"
+#include "Items/NavigationToolItemUtils.h"
+#include "NavigationToolView.h"
+#include "Widgets/Input/SEditableText.h"
+#include "Widgets/SNavigationToolTreeRow.h"
+
+#define LOCTEXT_NAMESPACE "SNavigationToolId"
+
+namespace UE::SequenceNavigator
+{
+
+using namespace Sequencer;
+
+void SNavigationToolId::Construct(const FArguments& InArgs
+	, const FNavigationToolViewModelPtr& InItem
+	, const TSharedRef<INavigationToolView>& InView
+	, const TSharedRef<SNavigationToolTreeRow>& InRowWidget)
+{
+	check(InItem.IsValid());
+	WeakItem = InItem;
+
+	ChildSlot
+	[
+		SNew(SEditableText)
+			.IsReadOnly(true)
+			.Text(this, &SNavigationToolId::GetItemText)
+	];
+}
+
+FText SNavigationToolId::GetItemText() const
+{
+	if (const FNavigationToolViewModelPtr Item = WeakItem.Pin())
+	{
+		if (const TViewModelPtr<IIdExtension> IdExtension = Item.ImplicitCast())
+		{
+			return IdExtension->GetId();
+		}
+	}
+	return FText();
+}
+
+} // namespace UE::SequenceNavigator
+
+#undef LOCTEXT_NAMESPACE
diff --git a/Engine/Plugins/Experimental/SequenceNavigator/Source/SequenceNavigator/Private/Widgets/Columns/SNavigationToolId.h b/Engine/Plugins/Experimental/SequenceNavigator/Source/SequenceNavigator/Private/Widgets/Columns/SNavigationToolId.h
new file mode 100644
--- /dev/null
+++ b/Engine/Plugins/Experimental/SequenceNavigator/Source/SequenceNavigator/Private/Widgets/Columns/SNavigationToolId.h
@@ -0,0 +1,34 @@
+// Copyright Epic Games, Inc. All Rights Reserved.
+
+#pragma once
+
+#include "NavigationToolDefines.h"
+#include "Widgets/SCompoundWidget.h"
+
+namespace UE::SequenceNavigator
+{
+
+class INavigationToolView;
+class SNavigationToolTreeRow;
+
+class SNavigationToolId : public SCompoundWidget
+{
+public:
+	SLATE_BEGIN_ARGS(SNavigationToolId) {}
+	SLATE_END_ARGS()
+
+	virtual ~SNavigationToolId() override {}
+
+	void Construct(const FArguments& InArgs
+		, const FNavigationToolViewModelPtr& InItem
+		, const TSharedRef<INavigationToolView>& InView
+		, const TSharedRef<SNavigationToolTreeRow>& InRowWidget);
+
+protected:
+
+	FText GetItemText() const;
+
+	FNavigationToolViewModelWeakPtr WeakItem;
+};
+
+} // namespace UE::SequenceNavigator
diff --git a/Engine/Plugins/Experimental/SequenceNavigator/Source/SequenceNavigator/Public/Columns/NavigationToolIdColumn.h b/Engine/Plugins/Experimental/SequenceNavigator/Source/SequenceNavigator/Public/Columns/NavigationToolIdColumn.h
new file mode 100644
--- /dev/null
+++ b/Engine/Plugins/Experimental/SequenceNavigator/Source/SequenceNavigator/Public/Columns/NavigationToolIdColumn.h
@@ -0,0 +1,42 @@
+// Copyright Epic Games, Inc. All Rights Reserved.
+
+#pragma once
+
+#include "Columns/NavigationToolColumn.h"
+#include "Items/NavigationToolItem.h"
+#include "MVVM/ICastable.h"
+
+#define UE_API SEQUENCENAVIGATOR_API
+
+namespace UE::SequenceNavigator
+{
+
+class INavigationToolView;
+
+class FNavigationToolIdColumn
+	: public FNavigationToolColumn
+{
+public:
+	UE_SEQUENCER_DECLARE_CASTABLE_API(UE_API, FNavigationToolIdColumn
+		, FNavigationToolColumn)
+
+	static FName StaticColumnId() { return TEXT("Id"); }
+
+protected:
+	//~ Begin INavigationToolColumn
+	virtual FName GetColumnId() const override { return StaticColumnId(); }
+	UE_API virtual FText GetColumnDisplayNameText() const override;
+	UE_API virtual const FSlateBrush* GetIconBrush() const override;
+	virtual bool ShouldShowColumnByDefault() const override { return false; }
+	virtual float GetFillWidth() const override { return 5.f; }
+	UE_API virtual SHeaderRow::FColumn::FArguments ConstructHeaderRowColumn(const TSharedRef<INavigationToolView>& InView
+		, const float InFillSize) override;
+	UE_API virtual TSharedRef<SWidget> ConstructRowWidget(const FNavigationToolViewModelPtr& InItem
+		, const TSharedRef<INavigationToolView>& InView
+		, const TSharedRef<SNavigationToolTreeRow>& InRow) override;
+	//~ End INavigationToolColumn
+};
+
+} // namespace UE::SequenceNavigator
+
+#undef UE_API
diff --git a/Engine/Plugins/Experimental/SequenceNavigator/Source/SequenceNavigator/Public/Extensions/IIdExtension.h b/Engine/Plugins/Experimental/SequenceNavigator/Source/SequenceNavigator/Public/Extensions/IIdExtension.h
new file mode 100644
--- /dev/null
+++ b/Engine/Plugins/Experimental/SequenceNavigator/Source/SequenceNavigator/Public/Extensions/IIdExtension.h
@@ -0,0 +1,25 @@
+﻿// Copyright Epic Games, Inc. All Rights Reserved.
+
+#pragma once
+
+#include "Misc/Optional.h"
+#include "MVVM/ViewModelTypeID.h"
+
+#define UE_API SEQUENCENAVIGATOR_API
+
+namespace UE::SequenceNavigator
+{
+
+class IIdExtension
+{
+public:
+	UE_SEQUENCER_DECLARE_VIEW_MODEL_TYPE_ID_API(UE_API, IIdExtension)
+
+	virtual ~IIdExtension() = default;
+
+	virtual FText GetId() const = 0;
+};
+
+} // namespace UE::SequenceNavigator
+
+#undef UE_API
diff --git a/Engine/Plugins/Experimental/SequenceNavigator/Source/SequenceNavigator/Public/Items/NavigationToolBinding.h b/Engine/Plugins/Experimental/SequenceNavigator/Source/SequenceNavigator/Public/Items/NavigationToolBinding.h
--- a/Engine/Plugins/Experimental/SequenceNavigator/Source/SequenceNavigator/Public/Items/NavigationToolBinding.h
+++ b/Engine/Plugins/Experimental/SequenceNavigator/Source/SequenceNavigator/Public/Items/NavigationToolBinding.h
@@ -2,6 +2,7 @@
 
 #pragma once
 
+#include "Extensions/IIdExtension.h"
 #include "Extensions/IPlayheadExtension.h"
 #include "MovieSceneBinding.h"
 #include "MVVM/Extensions/IRenameableExtension.h"
@@ -31,13 +32,15 @@ class FNavigationToolBinding
 	, public Sequencer::IRenameableExtension
 	//, public ISequenceLockableExtension
 	, public IPlayheadExtension
+	, public IIdExtension
 {
 public:
 	UE_SEQUENCER_DECLARE_CASTABLE_API(UE_API, FNavigationToolBinding
 		, FNavigationToolItem
 		, Sequencer::IRenameableExtension
 		//, ISequenceLockableExtension
-		, IPlayheadExtension);
+		, IPlayheadExtension
+		, IIdExtension);
 
 	UE_API FNavigationToolBinding(INavigationTool& InTool
 		, const FNavigationToolViewModelPtr& InParentItem
@@ -83,6 +86,10 @@ public:
 	//virtual void SetIsLocked(const bool bInIsLocked) override;
 	//~ End ISequenceLockableExtension
 
+	//~ Begin IIdExtension
+	UE_API virtual FText GetId() const override;
+	//~ End IIdExtension
+
 	//~ Begin IPlayheadExtension
 	UE_API virtual EItemContainsPlayhead ContainsPlayhead() const override;
 	//~ End IPlayheadExtension
diff --git a/Engine/Plugins/Experimental/SequenceNavigator/Source/SequenceNavigator/Public/Items/NavigationToolSequence.h b/Engine/Plugins/Experimental/SequenceNavigator/Source/SequenceNavigator/Public/Items/NavigationToolSequence.h
--- a/Engine/Plugins/Experimental/SequenceNavigator/Source/SequenceNavigator/Public/Items/NavigationToolSequence.h
+++ b/Engine/Plugins/Experimental/SequenceNavigator/Source/SequenceNavigator/Public/Items/NavigationToolSequence.h
@@ -3,6 +3,7 @@
 #pragma once
 
 #include "Extensions/IInTimeExtension.h"
+#include "Extensions/IIdExtension.h"
 #include "Extensions/IMarkerVisibilityExtension.h"
 #include "Extensions/IOutTimeExtension.h"
 #include "Extensions/IPlayheadExtension.h"
@@ -39,6 +40,7 @@ class FNavigationToolSequence
 	, public IMarkerVisibilityExtension
 	, public IRevisionControlExtension
 	, public IPlayheadExtension
+	, public IIdExtension
 	, public IInTimeExtension
 	, public IOutTimeExtension
 {
@@ -51,6 +53,7 @@ public:
 		, IMarkerVisibilityExtension
 		, IRevisionControlExtension
 		, IPlayheadExtension
+		, IIdExtension
 		, IInTimeExtension
 		, IOutTimeExtension)
 
@@ -115,6 +118,10 @@ public:
 	UE_API virtual void SetColor(const TOptional<FColor>& InColor) override;
 	//~ End IColorExtension
 
+	//~ Begin IIdExtension
+	UE_API virtual FText GetId() const override;
+	//~ End IIdExtension
+
 	//~ Begin IRevisionControlExtension
 	UE_API virtual EItemRevisionControlState GetRevisionControlState() const override;
 	UE_API virtual const FSlateBrush* GetRevisionControlStatusIcon() const override;
