[SequenceNavigator] Remove EditAnywhere and Category UPROPERTY specifiers from UNavigationToolSettings::FilterBars property to hide internally saved filter bar configs from showing in editor settings

diff --git a/Engine/Plugins/Experimental/SequenceNavigator/Source/SequenceNavigator/Public/NavigationToolSettings.h b/Engine/Plugins/Experimental/SequenceNavigator/Source/SequenceNavigator/Public/NavigationToolSettings.h
--- a/Engine/Plugins/Experimental/SequenceNavigator/Source/SequenceNavigator/Public/NavigationToolSettings.h
+++ b/Engine/Plugins/Experimental/SequenceNavigator/Source/SequenceNavigator/Public/NavigationToolSettings.h
@@ -148,7 +148,7 @@ private:
 	TSet<FName> EnabledBuiltInFilters;
 
 	/** Saved settings for each unique filter bar instance mapped by instance identifier */
-	UPROPERTY(config, EditAnywhere, Category = "Filtering")
+	UPROPERTY(Config)
 	TMap<FName, FSequencerFilterBarConfig> FilterBars;
 
 	/** Automatically expand tracks that pass filters */
