[SequenceNavigator] Correct playhead column display for bound items

diff --git a/Engine/Plugins/Experimental/SequenceNavigator/Source/SequenceNavigator/Private/Items/NavigationToolBinding.cpp b/Engine/Plugins/Experimental/SequenceNavigator/Source/SequenceNavigator/Private/Items/NavigationToolBinding.cpp
--- a/Engine/Plugins/Experimental/SequenceNavigator/Source/SequenceNavigator/Private/Items/NavigationToolBinding.cpp
+++ b/Engine/Plugins/Experimental/SequenceNavigator/Source/SequenceNavigator/Private/Items/NavigationToolBinding.cpp
@@ -262,6 +262,14 @@ FText FNavigationToolBinding::GetId() const
 
 EItemContainsPlayhead FNavigationToolBinding::ContainsPlayhead() const
 {
+	for (const FNavigationToolViewModelWeakPtr& WeakItem : GetChildren())
+	{
+		if (const TViewModelPtr<IPlayheadExtension> PlayheadExtension = WeakItem.ImplicitPin())
+		{
+			return PlayheadExtension->ContainsPlayhead();
+		}
+	}
+
 	const TSharedPtr<ISequencer> Sequencer = Tool.GetSequencer();
 	if (!Sequencer.IsValid())
 	{
@@ -269,13 +277,15 @@ EItemContainsPlayhead FNavigationToolBinding::ContainsPlayhead() const
 	}
 
 	const FQualifiedFrameTime PlayheadTime = Sequencer->GetLocalTime();
+	const FFrameNumber PlayheadFrame = PlayheadTime.ConvertTo(Sequencer->GetFocusedTickResolution()).FloorToFrame();
 
 	for (const UMovieSceneTrack* const Track : Binding.GetTracks())
 	{
 		const TArray<UMovieSceneSection*> Sections = Track->GetAllSections();
 		for (const UMovieSceneSection* const Section : Sections)
 		{
-			if (Section->GetTrueRange().Contains(PlayheadTime.Time.FrameNumber))
+			const TRange<FFrameNumber> Range = Section->ComputeEffectiveRange();
+			if (Range.Contains(PlayheadFrame))
 			{
 				return EItemContainsPlayhead::ContainsPlayhead;
 			}
diff --git a/Engine/Plugins/Experimental/SequenceNavigator/Source/SequenceNavigator/Private/Items/NavigationToolTrack.cpp b/Engine/Plugins/Experimental/SequenceNavigator/Source/SequenceNavigator/Private/Items/NavigationToolTrack.cpp
--- a/Engine/Plugins/Experimental/SequenceNavigator/Source/SequenceNavigator/Private/Items/NavigationToolTrack.cpp
+++ b/Engine/Plugins/Experimental/SequenceNavigator/Source/SequenceNavigator/Private/Items/NavigationToolTrack.cpp
@@ -13,6 +13,7 @@
 #include "MVVM/ViewModels/TrackModel.h"
 #include "NavigationToolScopedSelection.h"
 #include "ScopedTransaction.h"
+#include "Sections/MovieSceneSubSection.h"
 #include "Styling/SlateIconFinder.h"
 #include "Utils/NavigationToolMiscUtils.h"
 #include "Widgets/Views/STableRow.h"
@@ -223,10 +224,11 @@ EItemContainsPlayhead FNavigationToolTrack::ContainsPlayhead() const
 	}
 
 	const FQualifiedFrameTime PlayheadTime = Sequencer->GetLocalTime();
+	const FFrameNumber PlayheadFrame = PlayheadTime.ConvertTo(Sequencer->GetFocusedTickResolution()).FloorToFrame();
 
 	if (WeakSection.IsValid())
 	{
-		return WeakSection.Pin()->GetTrueRange().Contains(PlayheadTime.Time.FrameNumber)
+		return WeakSection.Pin()->ComputeEffectiveRange().Contains(PlayheadFrame)
 			? EItemContainsPlayhead::ContainsPlayhead : EItemContainsPlayhead::None;
 	}
 
@@ -239,7 +241,7 @@ EItemContainsPlayhead FNavigationToolTrack::ContainsPlayhead() const
 	const TArray<UMovieSceneSection*> Sections = Track->GetAllSections();
 	for (const UMovieSceneSection* const Section : Sections)
 	{
-		if (Section->GetTrueRange().Contains(PlayheadTime.Time.FrameNumber))
+		if (Section->ComputeEffectiveRange().Contains(PlayheadTime.Time.FrameNumber))
 		{
 			return EItemContainsPlayhead::ContainsPlayhead;
 		}
