[SequenceNavigator] Use the converted PladheadFrame instead of the PlayheadTime frame number

diff --git a/Engine/Plugins/Experimental/SequenceNavigator/Source/SequenceNavigator/Private/Items/NavigationToolTrack.cpp b/Engine/Plugins/Experimental/SequenceNavigator/Source/SequenceNavigator/Private/Items/NavigationToolTrack.cpp
--- a/Engine/Plugins/Experimental/SequenceNavigator/Source/SequenceNavigator/Private/Items/NavigationToolTrack.cpp
+++ b/Engine/Plugins/Experimental/SequenceNavigator/Source/SequenceNavigator/Private/Items/NavigationToolTrack.cpp
@@ -241,7 +241,7 @@ EItemContainsPlayhead FNavigationToolTrack::ContainsPlayhead() const
 	const TArray<UMovieSceneSection*> Sections = Track->GetAllSections();
 	for (const UMovieSceneSection* const Section : Sections)
 	{
-		if (Section->ComputeEffectiveRange().Contains(PlayheadTime.Time.FrameNumber))
+		if (Section->ComputeEffectiveRange().Contains(PlayheadFrame))
 		{
 			return EItemContainsPlayhead::ContainsPlayhead;
 		}
