[SequenceNavigator] Add usage analytics for 5.7

diff --git a/Engine/Plugins/Experimental/SequenceNavigator/Source/SequenceNavigator/Private/NavigationTool.cpp b/Engine/Plugins/Experimental/SequenceNavigator/Source/SequenceNavigator/Private/NavigationTool.cpp
--- a/Engine/Plugins/Experimental/SequenceNavigator/Source/SequenceNavigator/Private/NavigationTool.cpp
+++ b/Engine/Plugins/Experimental/SequenceNavigator/Source/SequenceNavigator/Private/NavigationTool.cpp
@@ -30,6 +30,7 @@
 #include "Providers/NavigationToolProvider.h"
 #include "ScopedTransaction.h"
 #include "Sections/MovieSceneSubSection.h"
+#include "Utils/NavigationToolMiscUtils.h"
 #include "Widgets/Views/STableRow.h"
 
 #define LOCTEXT_NAMESPACE "NavigationTool"
@@ -1071,6 +1072,8 @@ void FNavigationTool::ShowHideToolTab(const bool bInVisible)
 	if (bInVisible)
 	{
 		BindEvents();
+
+		RecordUsageAnalytics(TEXT("ToolOpened"), TEXT("true"));
 	}
 	else
 	{
diff --git a/Engine/Plugins/Experimental/SequenceNavigator/Source/SequenceNavigator/Private/Utils/NavigationToolMiscUtils.cpp b/Engine/Plugins/Experimental/SequenceNavigator/Source/SequenceNavigator/Private/Utils/NavigationToolMiscUtils.cpp
--- a/Engine/Plugins/Experimental/SequenceNavigator/Source/SequenceNavigator/Private/Utils/NavigationToolMiscUtils.cpp
+++ b/Engine/Plugins/Experimental/SequenceNavigator/Source/SequenceNavigator/Private/Utils/NavigationToolMiscUtils.cpp
@@ -1,6 +1,7 @@
 ﻿// Copyright Epic Games, Inc. All Rights Reserved.
 
 #include "NavigationToolMiscUtils.h"
+#include "EngineAnalytics.h"
 #include "INavigationTool.h"
 #include "ISourceControlModule.h"
 #include "ISourceControlProvider.h"
@@ -21,6 +22,15 @@ namespace UE::SequenceNavigator
 
 using namespace Sequencer;
 
+void RecordUsageAnalytics(const FString& InParamName, const FString& InParamValue)
+{
+	if (FEngineAnalytics::IsAvailable())
+	{
+		FEngineAnalytics::GetProvider().RecordEvent(TEXT("Editor.Usage.SequenceNavigator")
+			, InParamName, InParamValue);
+	}
+}
+
 FSourceControlStatePtr FindSourceControlState(UPackage* const InPackage)
 {
 	if (!InPackage)
@@ -114,6 +124,8 @@ void FocusSequence(const INavigationTool& InTool
 			Sequencer->SelectSection(SubSection);
 		}
 	}
+
+	RecordUsageAnalytics(TEXT("FocusItemInSequencer"), TEXT("SubSection"));
 }
 
 void FocusSequence(const INavigationTool& InTool
@@ -141,6 +153,8 @@ void FocusSequence(const INavigationTool& InTool
 			}
 		}
 	}
+
+	RecordUsageAnalytics(TEXT("FocusItemInSequencer"), TEXT("Track"));
 }
 
 void FocusSequence(const INavigationTool& InTool
@@ -164,6 +178,8 @@ void FocusSequence(const INavigationTool& InTool
 			}
 		}
 	}
+
+	RecordUsageAnalytics(TEXT("FocusItemInSequencer"), TEXT("ObjectBinding"));
 }
 
 void SelectSequencerBindingTrack(const ISequencer& InSequencer, const FGuid& InObjectId)
diff --git a/Engine/Plugins/Experimental/SequenceNavigator/Source/SequenceNavigator/Private/Utils/NavigationToolMiscUtils.h b/Engine/Plugins/Experimental/SequenceNavigator/Source/SequenceNavigator/Private/Utils/NavigationToolMiscUtils.h
--- a/Engine/Plugins/Experimental/SequenceNavigator/Source/SequenceNavigator/Private/Utils/NavigationToolMiscUtils.h
+++ b/Engine/Plugins/Experimental/SequenceNavigator/Source/SequenceNavigator/Private/Utils/NavigationToolMiscUtils.h
@@ -22,6 +22,13 @@ namespace UE::SequenceNavigator
 
 class INavigationTool;
 
+/**
+ * Records usage analytics for the Sequence Navigator tool.
+ * @param InParamName Parameter name string
+ * @param InParamValue Parameter value string
+ */
+void RecordUsageAnalytics(const FString& InParamName, const FString& InParamValue);
+
 /**
  * Find the source control state for a given package.
  * @param InPackage The package for which the source control state is being queried.
