Sequencer: Fix Unity build issues

#Synced-CL 49029986

diff --git a/Engine/Plugins/Experimental/SequenceNavigator/Source/SequenceNavigator/Private/Items/NavigationToolItemId.cpp b/Engine/Plugins/Experimental/SequenceNavigator/Source/SequenceNavigator/Private/Items/NavigationToolItemId.cpp
--- a/Engine/Plugins/Experimental/SequenceNavigator/Source/SequenceNavigator/Private/Items/NavigationToolItemId.cpp
+++ b/Engine/Plugins/Experimental/SequenceNavigator/Source/SequenceNavigator/Private/Items/NavigationToolItemId.cpp
@@ -1,4 +1,4 @@
-﻿// Copyright Epic Games, Inc. All Rights Reserved.
+// Copyright Epic Games, Inc. All Rights Reserved.
 
 #include "Items/NavigationToolItemId.h"
 #include "ItemProxies/INavigationToolItemProxyFactory.h"
@@ -11,6 +11,7 @@
 #include "MVVM/ViewModels/SectionModel.h"
 #include "MVVM/ViewModels/SequenceModel.h"
 #include "UObject/SoftObjectPath.h"
+#include "MVVM/ViewModels/ViewModelIterators.h"
 
 namespace UE::SequenceNavigator
 {
diff --git a/Engine/Source/Editor/MovieSceneTools/Private/IMovieSceneAnimMixerItemMenuProvider.cpp b/Engine/Source/Editor/MovieSceneTools/Private/IMovieSceneAnimMixerItemMenuProvider.cpp
new file mode 100644
--- /dev/null
+++ b/Engine/Source/Editor/MovieSceneTools/Private/IMovieSceneAnimMixerItemMenuProvider.cpp
@@ -0,0 +1,14 @@
+// Copyright Epic Games, Inc. All Rights Reserved.
+
+#include "IMovieSceneAnimMixerItemMenuProvider.h"
+
+IMovieSceneAnimMixerItemMenuProvider::IMovieSceneAnimMixerItemMenuProvider()
+{
+
+}
+
+FName IMovieSceneAnimMixerItemMenuProvider::GetModularFeatureName()
+{
+	static FName FeatureName = FName(TEXT("MovieSceneAnimMixerItemProvider"));
+	return FeatureName;
+}
\ No newline at end of file
diff --git a/Engine/Source/Editor/MovieSceneTools/Public/IMovieSceneAnimMixerItemMenuProvider.h b/Engine/Source/Editor/MovieSceneTools/Public/IMovieSceneAnimMixerItemMenuProvider.h
--- a/Engine/Source/Editor/MovieSceneTools/Public/IMovieSceneAnimMixerItemMenuProvider.h
+++ b/Engine/Source/Editor/MovieSceneTools/Public/IMovieSceneAnimMixerItemMenuProvider.h
@@ -20,11 +20,10 @@ struct FGuid;
 class MOVIESCENETOOLS_API IMovieSceneAnimMixerItemMenuProvider : public IModularFeature
 {
 public:
-	static FName GetModularFeatureName()
-	{
-		static FName FeatureName = FName(TEXT("MovieSceneAnimMixerItemProvider"));
-		return FeatureName;
-	}
+
+	static FName GetModularFeatureName();
+
+	IMovieSceneAnimMixerItemMenuProvider();
 
 	virtual ~IMovieSceneAnimMixerItemMenuProvider() = default;
 
diff --git a/Engine/Source/Editor/Sequencer/Private/MVVM/Views/SOutlinerTrackColorPicker.cpp b/Engine/Source/Editor/Sequencer/Private/MVVM/Views/SOutlinerTrackColorPicker.cpp
--- a/Engine/Source/Editor/Sequencer/Private/MVVM/Views/SOutlinerTrackColorPicker.cpp
+++ b/Engine/Source/Editor/Sequencer/Private/MVVM/Views/SOutlinerTrackColorPicker.cpp
@@ -1,6 +1,7 @@
 // Copyright Epic Games, Inc. All Rights Reserved.
 
 #include "MVVM/Views/SOutlinerTrackColorPicker.h"
+#include "MVVM/Extensions/IOutlinerExtension.h"
 #include "MVVM/Extensions/ITrackExtension.h"
 #include "MVVM/ViewModels/ViewModel.h"
 #include "MVVM/ViewModels/ViewModelIterators.h"
diff --git a/Engine/Source/Editor/Sequencer/Private/SequencerSectionPainter.cpp b/Engine/Source/Editor/Sequencer/Private/SequencerSectionPainter.cpp
--- a/Engine/Source/Editor/Sequencer/Private/SequencerSectionPainter.cpp
+++ b/Engine/Source/Editor/Sequencer/Private/SequencerSectionPainter.cpp
@@ -6,6 +6,7 @@
 #include "MVVM/ViewModels/SectionModel.h"
 #include "MVVM/Extensions/IOutlinerExtension.h"
 #include "MVVM/Extensions/ITrackExtension.h"
+#include "MVVM/ViewModels/ViewModelIterators.h"
 
 FSequencerSectionPainter::FSequencerSectionPainter(FSlateWindowElementList& OutDrawElements, const FGeometry& InSectionGeometry, TSharedPtr<UE::Sequencer::FSectionModel> InSection)
 	: SectionModel(InSection)
