Sequencer: Save config when done resizing and not during

#Synced-CL 49189700

diff --git a/Engine/Plugins/Experimental/SequenceNavigator/Source/SequenceNavigator/Private/NavigationToolSettings.cpp b/Engine/Plugins/Experimental/SequenceNavigator/Source/SequenceNavigator/Private/NavigationToolSettings.cpp
--- a/Engine/Plugins/Experimental/SequenceNavigator/Source/SequenceNavigator/Private/NavigationToolSettings.cpp
+++ b/Engine/Plugins/Experimental/SequenceNavigator/Source/SequenceNavigator/Private/NavigationToolSettings.cpp
@@ -253,12 +253,16 @@ float UNavigationToolSettings::GetLastFilterBarSizeCoefficient() const
 	return FMath::Max<float>(LastFilterBarSizeCoefficient, 0.05f);
 }
 
-void UNavigationToolSettings::SetLastFilterBarSizeCoefficient(const float bInSizeCoefficient)
+void UNavigationToolSettings::SetLastFilterBarSizeCoefficient(const float bInSizeCoefficient, const bool bInSaveConfig)
 {
 	if (LastFilterBarSizeCoefficient != bInSizeCoefficient)
 	{
 		LastFilterBarSizeCoefficient = bInSizeCoefficient;
-		SaveConfig();
+
+		if (bInSaveConfig)
+		{
+			SaveConfig();
+		}
 	}
 }
 
diff --git a/Engine/Plugins/Experimental/SequenceNavigator/Source/SequenceNavigator/Private/Widgets/SNavigationToolView.cpp b/Engine/Plugins/Experimental/SequenceNavigator/Source/SequenceNavigator/Private/Widgets/SNavigationToolView.cpp
--- a/Engine/Plugins/Experimental/SequenceNavigator/Source/SequenceNavigator/Private/Widgets/SNavigationToolView.cpp
+++ b/Engine/Plugins/Experimental/SequenceNavigator/Source/SequenceNavigator/Private/Widgets/SNavigationToolView.cpp
@@ -1,4 +1,4 @@
-﻿// Copyright Epic Games, Inc. All Rights Reserved.
+// Copyright Epic Games, Inc. All Rights Reserved.
 
 #include "SNavigationToolView.h"
 #include "Filters/Widgets/SFilterBarClippingHorizontalBox.h"
@@ -196,6 +196,13 @@ TSharedRef<SWidget> SNavigationToolView::ConstructSplitterContent()
 
 	return SNew(SSplitter)
 		.Orientation(Orient_Horizontal)
+		.OnSplitterFinishedResizing_Lambda([this]()
+			{
+				if (UNavigationToolSettings* const ToolSettings = GetMutableDefault<UNavigationToolSettings>())
+				{
+					ToolSettings->SaveConfig();
+				}
+			})
 		+ SSplitter::Slot()
 		.Value(TAttribute<float>::CreateLambda([this]()
 			{
@@ -209,7 +216,7 @@ TSharedRef<SWidget> SNavigationToolView::ConstructSplitterContent()
 			{
 				if (UNavigationToolSettings* const ToolSettings = GetMutableDefault<UNavigationToolSettings>())
 				{
-					ToolSettings->SetLastFilterBarSizeCoefficient(InNewCoefficient);
+					ToolSettings->SetLastFilterBarSizeCoefficient(InNewCoefficient, false);
 				}
 			})
 		[
diff --git a/Engine/Plugins/Experimental/SequenceNavigator/Source/SequenceNavigator/Public/NavigationToolSettings.h b/Engine/Plugins/Experimental/SequenceNavigator/Source/SequenceNavigator/Public/NavigationToolSettings.h
--- a/Engine/Plugins/Experimental/SequenceNavigator/Source/SequenceNavigator/Public/NavigationToolSettings.h
+++ b/Engine/Plugins/Experimental/SequenceNavigator/Source/SequenceNavigator/Public/NavigationToolSettings.h
@@ -104,7 +104,7 @@ public:
 	void SetFilterBarLayout(const EFilterBarLayout InLayout);
 
 	float GetLastFilterBarSizeCoefficient() const;
-	void SetLastFilterBarSizeCoefficient(const float bInSizeCoefficient);
+	void SetLastFilterBarSizeCoefficient(const float bInSizeCoefficient, const bool bInSaveConfig);
 
 	bool ShouldSyncSelectionToNavigationTool() const;
 	void SetSyncSelectionToNavigationTool(const bool bInSync, const bool bInSaveConfig = true);
diff --git a/Engine/Source/Editor/Sequencer/Private/SSequencer.cpp b/Engine/Source/Editor/Sequencer/Private/SSequencer.cpp
--- a/Engine/Source/Editor/Sequencer/Private/SSequencer.cpp
+++ b/Engine/Source/Editor/Sequencer/Private/SSequencer.cpp
@@ -552,6 +552,11 @@ TSharedRef<SWidget> SSequencer::ConstructFilterBarContent()
 
 	return SNew(SSplitter)
 		.Orientation(Orient_Horizontal)
+		.OnSplitterFinishedResizing_Lambda([this]()
+			{
+				GetSequencerSettings()->SaveConfig();
+			})
+
 		+ SSplitter::Slot()
 		.Value(TAttribute<float>::CreateLambda([this]()
 			{
@@ -565,7 +570,7 @@ TSharedRef<SWidget> SSequencer::ConstructFilterBarContent()
 			}))
 		.OnSlotResized_Lambda([this](const float InNewCoefficient)
 			{
-				GetSequencerSettings()->SetLastFilterBarSizeCoefficient(InNewCoefficient);
+				GetSequencerSettings()->SetLastFilterBarSizeCoefficient(InNewCoefficient, false);
 			})
 		[
 			SFilterBarClippingHorizontalBox::WrapVerticalListWithHeading(FilterBarWidgetRef
diff --git a/Engine/Source/Editor/Sequencer/Private/SequencerSettings.cpp b/Engine/Source/Editor/Sequencer/Private/SequencerSettings.cpp
--- a/Engine/Source/Editor/Sequencer/Private/SequencerSettings.cpp
+++ b/Engine/Source/Editor/Sequencer/Private/SequencerSettings.cpp
@@ -1250,10 +1250,17 @@ float USequencerSettings::GetLastFilterBarSizeCoefficient() const
 	return LastFilterBarSizeCoefficient;
 }
 
-void USequencerSettings::SetLastFilterBarSizeCoefficient(const float bInSizeCoefficient)
+void USequencerSettings::SetLastFilterBarSizeCoefficient(const float bInSizeCoefficient, const bool bInSaveConfig)
 {
-	LastFilterBarSizeCoefficient = bInSizeCoefficient;
-	SaveConfig();
+	if (LastFilterBarSizeCoefficient != bInSizeCoefficient)
+	{
+		LastFilterBarSizeCoefficient = bInSizeCoefficient;
+
+		if (bInSaveConfig)
+		{
+			SaveConfig();
+		}
+	}
 }
 
 void USequencerSettings::SetThumbnailCaptureSettings(const FSequencerThumbnailCaptureSettings& InNewValue)
diff --git a/Engine/Source/Editor/Sequencer/Public/SequencerSettings.h b/Engine/Source/Editor/Sequencer/Public/SequencerSettings.h
--- a/Engine/Source/Editor/Sequencer/Public/SequencerSettings.h
+++ b/Engine/Source/Editor/Sequencer/Public/SequencerSettings.h
@@ -565,7 +565,7 @@ public:
 	UE_API void SetFilterBarLayout(const EFilterBarLayout InLayout);
 
 	UE_API float GetLastFilterBarSizeCoefficient() const;
-	UE_API void SetLastFilterBarSizeCoefficient(const float bInSizeCoefficient);
+	UE_API void SetLastFilterBarSizeCoefficient(const float bInSizeCoefficient, const bool bInSaveConfig);
 	
 	/** Gets the settings that determine how the asset thumbnail is captured when the sequence is saved. */
 	const FSequencerThumbnailCaptureSettings& GetThumbnailCaptureSettings() const { return ThumbnailCaptureSettings; }
