[Sequencer] Fix "Name" text expression "!=" operator not working

#Synced-CL 50319781

diff --git a/Engine/Plugins/Experimental/SequenceNavigator/Source/SequenceNavigator/Private/Filters/TextExpressions/NavigationToolFilterTextExpression_Name.cpp b/Engine/Plugins/Experimental/SequenceNavigator/Source/SequenceNavigator/Private/Filters/TextExpressions/NavigationToolFilterTextExpression_Name.cpp
--- a/Engine/Plugins/Experimental/SequenceNavigator/Source/SequenceNavigator/Private/Filters/TextExpressions/NavigationToolFilterTextExpression_Name.cpp
+++ b/Engine/Plugins/Experimental/SequenceNavigator/Source/SequenceNavigator/Private/Filters/TextExpressions/NavigationToolFilterTextExpression_Name.cpp
@@ -2,6 +2,7 @@
 
 #include "NavigationToolFilterTextExpression_Name.h"
 #include "Items/NavigationToolItem.h"
+#include "MVVM/ViewModels/ViewModelIterators.h"
 
 #define LOCTEXT_NAMESPACE "NavigationToolFilterTextExpression_Name"
 
@@ -35,19 +36,61 @@ bool FNavigationToolFilterTextExpression_Name::TestComplexExpression(const FName
 	, const ETextFilterComparisonOperation InComparisonOperation
 	, const ETextFilterTextComparisonMode InTextComparisonMode) const
 {
+	/**
+	 * Evaluate a keyed text-filter expression against the "Name" of an item, with outliner/subtree semantics.
+	 *
+	 * Key points:
+	 * - This filter is evaluated against the labels of the current row AND its ancestors (bIncludeThis=true).
+	 *   That means a match on a parent label causes the child row to be considered a match as well.
+	 *
+	 * - Most comparison operators (==, contains, startswith, etc.) use "any-match" aggregation:
+	 *     Pass if ANY ancestor/self label matches the operator against InValue.
+	 *
+	 * - "!=" is intentionally special-cased to support subtree exclusion:
+	 *     The desired behavior is: "exclude this row and its subtree if ANY ancestor/self label equals the value".
+	 *     Said differently: for "name != X" we want to keep rows only when NO ancestor/self label equals X.
+	 *
+	 * Implemented as:
+	 *   1) For "!=" we run per-label comparisons using "==" (Equal) to detect "any equals".
+	 *   2) After the scan, we invert the aggregate result for "!=".
+	 *
+	 * Notes on the early return:
+	 * - The base context's TestComplexExpression(...) determines whether this expression applies to this filter key.
+	 *   If it does NOT apply, we return true so that unrelated keys don't accidentally filter out items.
+	 */
+
 	if (!FNavigationToolFilterTextExpressionContext::TestComplexExpression(InKey, InValue, InComparisonOperation, InTextComparisonMode))
 	{
 		return true;
 	}
 
-	if (const FNavigationToolViewModelPtr FilterItem = WeakFilterItem.Pin())
+	const TViewModelPtr<INavigationToolItem> FilterItem = WeakFilterItem.Pin();
+	if (!FilterItem.IsValid())
+	{
+		return false;
+	}
+
+	const bool bIsNotEqual = (InComparisonOperation == ETextFilterComparisonOperation::NotEqual);
+
+	// Subtree exclusion for "!=" is based on "any ancestor/self equals".
+	// For all other operators, apply the operator normally.
+	const ETextFilterComparisonOperation ActualOperation = bIsNotEqual
+		? ETextFilterComparisonOperation::Equal : InComparisonOperation;
+
+	bool bAnyMatched = false;
+
+	for (const TViewModelPtr<IOutlinerExtension>& OutlinerExtension
+		: FilterItem.AsModel()->GetAncestorsOfType<IOutlinerExtension>(/*bIncludeThis=*/true))
 	{
-		const FString Label = FilterItem->GetDisplayName().ToString();
-		const bool bPassed = TextFilterUtils::TestComplexExpression(Label, InValue, InComparisonOperation, InTextComparisonMode);
-		return bPassed;
+		const FString Label = OutlinerExtension->GetLabel().ToString();
+		if (TextFilterUtils::TestComplexExpression(Label, InValue, ActualOperation, InTextComparisonMode))
+		{
+			bAnyMatched = true;
+			break;
+		}
 	}
 
-	return false;
+	return bIsNotEqual ? !bAnyMatched : bAnyMatched;
 }
 
 } // namespace UE::SequenceNavigator
diff --git a/Engine/Source/Editor/Sequencer/Private/Filters/TextExpressions/SequencerTextFilterExpression_Name.cpp b/Engine/Source/Editor/Sequencer/Private/Filters/TextExpressions/SequencerTextFilterExpression_Name.cpp
--- a/Engine/Source/Editor/Sequencer/Private/Filters/TextExpressions/SequencerTextFilterExpression_Name.cpp
+++ b/Engine/Source/Editor/Sequencer/Private/Filters/TextExpressions/SequencerTextFilterExpression_Name.cpp
@@ -34,21 +34,55 @@ bool FSequencerTextFilterExpression_Name::TestComplexExpression(const FName& InK
 	, const ETextFilterComparisonOperation InComparisonOperation
 	, const ETextFilterTextComparisonMode InTextComparisonMode) const
 {
+	/**
+	 * Evaluate a keyed text-filter expression against the "Name" of an item, with outliner/subtree semantics.
+	 *
+	 * Key points:
+	 * - This filter is evaluated against the labels of the current row AND its ancestors (bIncludeThis=true).
+	 *   That means a match on a parent label causes the child row to be considered a match as well.
+	 *
+	 * - Most comparison operators (==, contains, startswith, etc.) use "any-match" aggregation:
+	 *     Pass if ANY ancestor/self label matches the operator against InValue.
+	 *
+	 * - "!=" is intentionally special-cased to support subtree exclusion:
+	 *     The desired behavior is: "exclude this row and its subtree if ANY ancestor/self label equals the value".
+	 *     Said differently: for "name != X" we want to keep rows only when NO ancestor/self label equals X.
+	 *
+	 * Implemented as:
+	 *   1) For "!=" we run per-label comparisons using "==" (Equal) to detect "any equals".
+	 *   2) After the scan, we invert the aggregate result for "!=".
+	 *
+	 * Notes on the early return:
+	 * - The base context's TestComplexExpression(...) determines whether this expression applies to this filter key.
+	 *   If it does NOT apply, we return true so that unrelated keys don't accidentally filter out items.
+	 */
+
 	if (!FSequencerTextFilterExpressionContext::TestComplexExpression(InKey, InValue, InComparisonOperation, InTextComparisonMode))
 	{
 		return true;
 	}
 
-	for (const TViewModelPtr<IOutlinerExtension>& OutlinerExtension : FilterItem->GetAncestorsOfType<IOutlinerExtension>())
+	const bool bIsNotEqual = (InComparisonOperation == ETextFilterComparisonOperation::NotEqual);
+
+	// Subtree exclusion for "!=" is based on "any ancestor/self equals".
+	// For all other operators, apply the operator normally.
+	const ETextFilterComparisonOperation ActualOperation = bIsNotEqual
+		? ETextFilterComparisonOperation::Equal : InComparisonOperation;
+
+	bool bAnyMatched = false;
+
+	for (const TViewModelPtr<IOutlinerExtension>& OutlinerExtension
+		: FilterItem->GetAncestorsOfType<IOutlinerExtension>(/*bIncludeThis=*/true))
 	{
 		const FString Label = OutlinerExtension->GetLabel().ToString();
-		if (TextFilterUtils::TestComplexExpression(Label, InValue, InComparisonOperation, InTextComparisonMode))
+		if (TextFilterUtils::TestComplexExpression(Label, InValue, ActualOperation, InTextComparisonMode))
 		{
-			return true;
+			bAnyMatched = true;
+			break;
 		}
 	}
 
-	return false;
+	return bIsNotEqual ? !bAnyMatched : bAnyMatched;
 }
 
 #undef LOCTEXT_NAMESPACE
