[SequencerSimpleView] Force Sequence Navigator button to be immediately after the Curve Editor button on the toolbar
#Synced-CL 52691839

diff --git a/Engine/Plugins/Experimental/SequenceNavigator/Source/SequenceNavigator/Private/NavigationToolExtender.cpp b/Engine/Plugins/Experimental/SequenceNavigator/Source/SequenceNavigator/Private/NavigationToolExtender.cpp
--- a/Engine/Plugins/Experimental/SequenceNavigator/Source/SequenceNavigator/Private/NavigationToolExtender.cpp
+++ b/Engine/Plugins/Experimental/SequenceNavigator/Source/SequenceNavigator/Private/NavigationToolExtender.cpp
@@ -15,6 +15,8 @@
 #include "SequencerSettings.h"
 #include "SequencerToolMenuContext.h"
 
+#define LOCTEXT_NAMESPACE "NavigationToolExtender"
+
 namespace UE::SequenceNavigator
 {
 
@@ -420,7 +422,9 @@ void FNavigationToolExtender::AddSequencerToolBarExtension(const FNavigationTool
 
 	const FName ToolId = InToolInstance.ToolId;
 
-	FToolMenuSection& NavigationToolSection = ToolMenu->FindOrAddSection(TEXT("NavigationTool"));
+	FToolMenuSection& NavigationToolSection = ToolMenu->FindOrAddSection(TEXT("NavigationTool")
+		, LOCTEXT("SequenceNavigatorSectionDisplayName", "Sequence Navigator")
+		, FToolMenuInsert(TEXT("CurveEditor"), EToolMenuInsertType::After));
 
 	NavigationToolSection.AddDynamicEntry(TEXT("NavigationTool")
 		, FNewToolMenuSectionDelegate::CreateLambda([ToolId, ToolInstance = InToolInstance](FToolMenuSection& InSection)
@@ -542,3 +546,5 @@ TSharedPtr<INavigationToolIconCustomization> FNavigationToolExtender::GetCustomi
 }
 
 } // namespace UE::SequenceNavigator
+
+#undef LOCTEXT_NAMESPACE
